Add basic timecode support by sending the user defined time code value over message bus and making the Live Link Source activate that provider in editor.

diff --git a/Engine/Plugins/Animation/LiveLink/Source/LiveLink/Public/LiveLinkTimecodeProvider.h b/Engine/Plugins/Animation/LiveLink/Source/LiveLink/Public/LiveLinkTimecodeProvider.h
--- a/Engine/Plugins/Animation/LiveLink/Source/LiveLink/Public/LiveLinkTimecodeProvider.h
+++ b/Engine/Plugins/Animation/LiveLink/Source/LiveLink/Public/LiveLinkTimecodeProvider.h
@@ -54,6 +54,15 @@ public:
 #endif
 	//~End UObject Interface
 
+	/**
+	 * Allows users to override the target timecode provider subject key. Live
+	 * Link Hub will send subject name to the host but the host has to "lookup"
+	 * and match it with the available subject keys on the local machine.
+	 */
+	void SetTargetSubjectKey(const FLiveLinkSubjectKey& InKey)
+	{
+		SubjectKey = InKey;
+	}
 private:
 	FQualifiedFrameTime ConvertTo(FQualifiedFrameTime Value) const;
 	FQualifiedFrameTime LerpBetweenFrames(double Seconds, int32 IndexA, int32 IndexB) const;
diff --git a/Engine/Plugins/Animation/LiveLinkHub/Source/LiveLinkHub/LiveLinkHub.Build.cs b/Engine/Plugins/Animation/LiveLinkHub/Source/LiveLinkHub/LiveLinkHub.Build.cs
--- a/Engine/Plugins/Animation/LiveLinkHub/Source/LiveLinkHub/LiveLinkHub.Build.cs
+++ b/Engine/Plugins/Animation/LiveLinkHub/Source/LiveLinkHub/LiveLinkHub.Build.cs
@@ -21,6 +21,7 @@ namespace UnrealBuildTool.Rules
 				"ContentBrowser",
 				"ContentBrowserAssetDataSource",
 				"ContentBrowserData",
+				"EditorStyle",
 				"Engine",
 				"InputCore",
 				"Json",
diff --git a/Engine/Plugins/Animation/LiveLinkHub/Source/LiveLinkHub/Private/Clients/LiveLinkHubProvider.cpp b/Engine/Plugins/Animation/LiveLinkHub/Source/LiveLinkHub/Private/Clients/LiveLinkHubProvider.cpp
--- a/Engine/Plugins/Animation/LiveLinkHub/Source/LiveLinkHub/Private/Clients/LiveLinkHubProvider.cpp
+++ b/Engine/Plugins/Animation/LiveLinkHub/Source/LiveLinkHub/Private/Clients/LiveLinkHubProvider.cpp
@@ -5,6 +5,7 @@
 #include "Algo/Transform.h"
 #include "Async/Async.h"
 #include "Clients/LiveLinkHubClientsModel.h"
+#include "Clients/LiveLinkHubProvider.h"
 #include "Clients/LiveLinkHubUEClientInfo.h"
 #include "Containers/ObservableArray.h"
 #include "CoreMinimal.h"
@@ -84,6 +85,17 @@ bool FLiveLinkHubProvider::ShouldTransmitToSubject_AnyThread(FName SubjectName,
 	return true;
 }
 
+void FLiveLinkHubProvider::SetTimecodeSettings(FLiveLinkHubTimecodeSettings InSettings)
+{
+	TimecodeSettings = MoveTemp(InSettings);
+	SendTimecodeSettings();
+}
+
+void FLiveLinkHubProvider::SendTimecodeSettings()
+{
+	SendMessage(FMessageEndpoint::MakeMessage<FLiveLinkHubTimecodeSettings>(TimecodeSettings));
+}
+
 void FLiveLinkHubProvider::AddRestoredClient(FLiveLinkHubUEClientInfo& RestoredClientInfo)
 {
 	// If a client was already discovered with the same hostname, update it to match the restored client.
@@ -230,6 +242,8 @@ void FLiveLinkHubProvider::HandleClientInfoMessage(const FLiveLinkClientInfoMess
 		}
 	}
 
+	SendTimecodeSettings();
+
 	if (ClientId.IsValid())
 	{
 		OnClientEventDelegate.Broadcast(ClientId, EClientEventType::Modified);
diff --git a/Engine/Plugins/Animation/LiveLinkHub/Source/LiveLinkHub/Private/Clients/LiveLinkHubProvider.h b/Engine/Plugins/Animation/LiveLinkHub/Source/LiveLinkHub/Private/Clients/LiveLinkHubProvider.h
--- a/Engine/Plugins/Animation/LiveLinkHub/Source/LiveLinkHub/Private/Clients/LiveLinkHubProvider.h
+++ b/Engine/Plugins/Animation/LiveLinkHub/Source/LiveLinkHub/Private/Clients/LiveLinkHubProvider.h
@@ -7,6 +7,7 @@
 #include "Clients/LiveLinkHubUEClientInfo.h"
 #include "Engine/TimerHandle.h"
 #include "IMessageContext.h"
+#include "LiveLinkHubMessages.h"
 #include "LiveLinkProviderImpl.h"
 
 class ILiveLinkHubSessionManager;
@@ -14,6 +15,7 @@ struct FLiveLinkHubClientId;
 struct FLiveLinkHubConnectMessage;
 struct FLiveLinkHubUEClientInfo;
 
+
 /** 
  * LiveLink Provider that allows getting more information about a UE client by communicating with a LiveLinkHub MessageBus Source.
  */
@@ -50,6 +52,9 @@ public:
 	/** Retrieve the existing client map. */
 	const TMap<FLiveLinkHubClientId, FLiveLinkHubUEClientInfo>& GetClientsMap() const { return ClientsMap; }
 
+	/** Timecode settings that should be shared to connected editors. */
+	void SetTimecodeSettings(FLiveLinkHubTimecodeSettings InSettings);
+
 private:
 	/** Handle a connection message resulting from a livelink hub message bus source connecting to this provider. */
 	void HandleHubConnectMessage(const FLiveLinkHubConnectMessage& Message, const TSharedRef<IMessageContext, ESPMode::ThreadSafe>& Context);
@@ -57,6 +62,9 @@ private:
 	/** Handle a client info message being received. Happens when new information about a client is received (ie. Client has changed map) */
 	void HandleClientInfoMessage(const FLiveLinkClientInfoMessage& Message, const TSharedRef<IMessageContext, ESPMode::ThreadSafe>& Context);
 
+	/** Send timecode settings to connected Live Link Hub provider. */
+	void SendTimecodeSettings();
+
 protected:
 	//~ Begin ILiveLinkHubClientsModel interface
 	virtual void OnConnectionsClosed(const TArray<FMessageAddress>& ClosedAddresses) override;
@@ -89,6 +97,9 @@ private:
 	TWeakPtr<ILiveLinkHubSessionManager> SessionManager;
 	/** Cache used to retrieve the client id from a message bus address. */
 	TMap<FMessageAddress, FLiveLinkHubClientId> AddressToIdCache;
+	/** Cached value of timecode connection settings. */
+	FLiveLinkHubTimecodeSettings TimecodeSettings;
+
 	/** Lock used to access the clients map from different threads. */
 	mutable FRWLock ClientsMapLock;
 };
diff --git a/Engine/Plugins/Animation/LiveLinkHub/Source/LiveLinkHub/Private/UI/Widgets/LiveLinkHubMainTabController.h b/Engine/Plugins/Animation/LiveLinkHub/Source/LiveLinkHub/Private/UI/Widgets/LiveLinkHubMainTabController.h
--- a/Engine/Plugins/Animation/LiveLinkHub/Source/LiveLinkHub/Private/UI/Widgets/LiveLinkHubMainTabController.h
+++ b/Engine/Plugins/Animation/LiveLinkHub/Source/LiveLinkHub/Private/UI/Widgets/LiveLinkHubMainTabController.h
@@ -4,8 +4,6 @@
 
 #include "UI/Window/ILiveLinkHubComponent.h"
 
-
-class ILiveLinkClient;
 class FSpawnTabArgs;
 class SLiveLinkHubMainTabView;
 class SWindow;
@@ -27,7 +25,4 @@ private:
 
 	/** Spawn the main tab. */
 	TSharedRef<SDockTab> SpawnMainTab(const FSpawnTabArgs& Args, TSharedPtr<SWindow> RootWindow);
-
-	/** LiveLink Client */
-	ILiveLinkClient* LiveLinkClient = nullptr;
 };
diff --git a/Engine/Plugins/Animation/LiveLinkHub/Source/LiveLinkHub/Private/UI/Widgets/SLiveLinkHubTabViewBase.cpp b/Engine/Plugins/Animation/LiveLinkHub/Source/LiveLinkHub/Private/UI/Widgets/SLiveLinkHubTabViewBase.cpp
--- a/Engine/Plugins/Animation/LiveLinkHub/Source/LiveLinkHub/Private/UI/Widgets/SLiveLinkHubTabViewBase.cpp
+++ b/Engine/Plugins/Animation/LiveLinkHub/Source/LiveLinkHub/Private/UI/Widgets/SLiveLinkHubTabViewBase.cpp
@@ -10,6 +10,8 @@
 #include "Modules/ModuleManager.h"
 #include "Recording/LiveLinkHubRecordingController.h"
 #include "SLiveLinkHubStatusBar.h"
+#include "SLiveLinkTimecode.h"
+
 #include "Widgets/Images/SImage.h"
 #include "Widgets/Input/SButton.h"
 #include "Widgets/Input/SComboButton.h"
@@ -21,6 +23,7 @@
 
 #define LOCTEXT_NAMESPACE "LiveLinkHub"
 
+#define LIVELINKHUB_SUPPORTS_LAYOUTS 0
 TSharedRef<SWidget> GetModeSwitcherContent()
 {
 	FMenuBuilder MenuBuilder(true, NULL);
@@ -56,6 +59,7 @@ void SLiveLinkHubTabViewBase::Construct(const FArguments& InArgs)
 				.Padding(FMargin(4.0f, 6.0f))
 				[
 					SNew(SHorizontalBox)
+#if LIVELINKHUB_SUPPORTS_LAYOUTS
 					+ SHorizontalBox::Slot()
 					.HAlign(HAlign_Left)
 					[
@@ -83,6 +87,7 @@ void SLiveLinkHubTabViewBase::Construct(const FArguments& InArgs)
 
 						]
 					]
+#endif
 					+ SHorizontalBox::Slot()
 					.HAlign(HAlign_Right)
 					.VAlign(VAlign_Center)
@@ -103,17 +108,12 @@ void SLiveLinkHubTabViewBase::Construct(const FArguments& InArgs)
 							SNew(SSeparator)
 							.Orientation(Orient_Vertical)
 						]
-						/* Disable until we use timecode in LLH
 						+ SHorizontalBox::Slot()
 						.Padding(2.f)
 	                    .VAlign(VAlign_Center)
 						[
-							SNew(STimecode)
-							.Timecode(MakeAttributeLambda([]
-							{
-								return FApp::GetTimecode();
-							}))
-						]*/
+							SNew(SLiveLinkTimecode)
+						]
 					]
 				]
 			]
diff --git a/Engine/Plugins/Animation/LiveLinkHub/Source/LiveLinkHub/Private/UI/Widgets/SLiveLinkTimecode.cpp b/Engine/Plugins/Animation/LiveLinkHub/Source/LiveLinkHub/Private/UI/Widgets/SLiveLinkTimecode.cpp
new file mode 100644
--- /dev/null
+++ b/Engine/Plugins/Animation/LiveLinkHub/Source/LiveLinkHub/Private/UI/Widgets/SLiveLinkTimecode.cpp
@@ -0,0 +1,232 @@
+// Copyright Epic Games, Inc. All Rights Reserved.
+
+#include "SLiveLinkTimecode.h"
+
+#include "UObject/NameTypes.h"
+
+#include "EditorFontGlyphs.h"
+#include "Features/IModularFeatures.h"
+#include "Fonts/SlateFontInfo.h"
+#include "Framework/MultiBox/MultiBoxBuilder.h"
+#include "Internationalization/Text.h"
+#include "Misc/App.h"
+#include "UI/Widgets/SLiveLinkTimecode.h"
+#include "Widgets/Input/SComboButton.h"
+#include "Widgets/STimecode.h"
+
+#include "Clients/LiveLinkHubProvider.h"
+#include "LiveLinkClient.h"
+#include "LiveLinkHub.h"
+#include "LiveLinkHubCommands.h"
+#include "LiveLinkHubMessages.h"
+#include "LiveLinkHubModule.h"
+
+#define LOCTEXT_NAMESPACE "LiveLinkHub"
+
+namespace UE::LiveLinkTimecode::Private
+{
+	static FName EnableTimecodeSourceId = TEXT("EnableTimeCodeSource");
+	/** We only support a preset list of timecode values + named subjects. */
+	static FName System24fps = TEXT("SystemTime24fps");
+	static FName System30fps = TEXT("SystemTime30fps");
+	static FName System60fps = TEXT("SystemTime60fps");
+}
+
+FSlateColor SLiveLinkTimecode::GetTimecodeStatusColor() const
+{
+	return bIsTimecodeSource ? FSlateColor(FColor::Green) : FSlateColor(FColor::Yellow);
+}
+
+FText SLiveLinkTimecode::GetTimecodeTooltip() const
+{
+	if (bIsTimecodeSource)
+	{
+		return LOCTEXT("LiveLinkTimeCode_IsConnected", "Sending timecode data to connected editors.");
+	}
+	return LOCTEXT("LiveLinkTimeCode_NotConnected", "No timecode data shared with connected editors.");
+}
+
+void SLiveLinkTimecode::OnEnableTimecodeToggled()
+{
+	bIsTimecodeSource = !bIsTimecodeSource;
+	SendUpdatedTimecodeToEditor();
+}
+
+void SLiveLinkTimecode::SendUpdatedTimecodeToEditor()
+{
+	FLiveLinkHubTimecodeSettings Settings;
+	using namespace UE::LiveLinkTimecode::Private;
+
+	if (!bIsTimecodeSource)
+	{
+		Settings.Source = ELiveLinkHubTimecodeSource::NotDefined;
+	}
+	else if (ActiveTimecodeSource == System24fps)
+	{
+		Settings.Source = ELiveLinkHubTimecodeSource::SystemTimeEditor;
+		Settings.DesiredFrameRate = FFrameRate(24,1);
+	}
+	else if (ActiveTimecodeSource == System30fps)
+	{
+		Settings.Source = ELiveLinkHubTimecodeSource::SystemTimeEditor;
+		Settings.DesiredFrameRate = FFrameRate(30, 1);
+	}
+	else if (ActiveTimecodeSource == System60fps)
+	{
+		Settings.Source = ELiveLinkHubTimecodeSource::SystemTimeEditor;
+		Settings.DesiredFrameRate = FFrameRate(60, 1);
+	}
+	else
+	{
+		Settings.Source = ELiveLinkHubTimecodeSource::UseSubjectName;
+	}
+	Settings.SubjectName = ActiveTimecodeSource;
+
+	// Apply the timecode settings to the local instance too.
+	Settings.AssignTimecodeSettingsAsProviderToEngine();
+
+	const TSharedPtr<FLiveLinkHub> LiveLinkHub = FModuleManager::Get().GetModuleChecked<FLiveLinkHubModule>("LiveLinkHub").GetLiveLinkHub();
+	if (LiveLinkHub.IsValid())
+	{
+		TSharedPtr<FLiveLinkHubProvider> Provider = LiveLinkHub->GetLiveLinkProvider();
+		if (Provider.IsValid())
+		{
+			Provider->SetTimecodeSettings(MoveTemp(Settings));
+		}
+	}
+}
+
+void SLiveLinkTimecode::SetTimecodeSource(const FName SourceId)
+{
+	if (SourceId == ActiveTimecodeSource)
+	{
+		return;
+	}
+
+	ActiveTimecodeSource = SourceId;
+	SendUpdatedTimecodeToEditor();
+}
+
+TSharedRef<SWidget> SLiveLinkTimecode::MakeMenu()
+{
+	FMenuBuilder MenuBuilder(/*bInShouldCloseWindowAfterMenuSelection=*/true, nullptr);
+
+	MenuBuilder.AddMenuEntry(
+		LOCTEXT("LiveLinkHubTimecodeSource", "Enable Timecode Source"),
+		LOCTEXT("LiveLinkHubTimecodeSource_Tooltip", "Make this Live Link Hub a time code source for connected editors."),
+		FSlateIcon(),
+		FUIAction(
+			FExecuteAction::CreateSP(this, &SLiveLinkTimecode::OnEnableTimecodeToggled),
+			FCanExecuteAction::CreateLambda([] { return true; }),
+			FIsActionChecked::CreateLambda([this] { return bIsTimecodeSource; })),
+		NAME_None,
+		EUserInterfaceActionType::ToggleButton);
+
+
+	MenuBuilder.BeginSection("LiveLinkHub.Timecode.TimecodeProvider", LOCTEXT("TimecodeProviderSection", "Timecode Provider"));
+
+	auto GenerateUIAction = [this](const FName Id) -> FUIAction
+	{
+		// Generate a UI Action
+		return FUIAction(
+			FExecuteAction::CreateSP(this, &SLiveLinkTimecode::SetTimecodeSource, Id),
+			FCanExecuteAction::CreateLambda([] { return true; }),
+			FIsActionChecked::CreateLambda([this, Id] { return Id == ActiveTimecodeSource; }));
+	};
+	MenuBuilder.AddMenuEntry(
+		LOCTEXT("LiveLinkHubTimecodeSource24fps", "System Time (24 fps)"),
+		LOCTEXT("LiveLinkHubTimecodeSource24fps_Tooltip", "Use a 24 FPS time code based on system time."),
+		FSlateIcon(),
+		GenerateUIAction(UE::LiveLinkTimecode::Private::System24fps),
+		NAME_None,
+		EUserInterfaceActionType::Check);
+
+	MenuBuilder.AddMenuEntry(
+		LOCTEXT("LiveLinkHubTimecodeSource30fps", "System Time (30 fps)"),
+		LOCTEXT("LiveLinkHubTimecodeSource30fps_Tooltip", "Use a 30 FPS time code based on system time."),
+		FSlateIcon(),
+		GenerateUIAction(UE::LiveLinkTimecode::Private::System30fps),
+		NAME_None,
+		EUserInterfaceActionType::Check);
+
+	MenuBuilder.AddMenuEntry(
+		LOCTEXT("LiveLinkHubTimecodeSource60fps", "System Time (60 fps)"),
+		LOCTEXT("LiveLinkHubTimecodeSource60fps_Tooltip", "Use a 60 FPS time code based on system time."),
+		FSlateIcon(),
+		GenerateUIAction(UE::LiveLinkTimecode::Private::System60fps),
+		NAME_None,
+		EUserInterfaceActionType::Check);
+
+	TArray<FName> Subjects;
+	WorkingClient->GetSubjectNames(Subjects);
+	for (FName Subject : Subjects)
+	{
+		FTextBuilder SubjectText;
+		SubjectText.AppendLineFormat(LOCTEXT("LiveLinkHubTimecodeSourceSubject_Tooltip", "{0}'s timecode"), FText::FromName(Subject));
+		MenuBuilder.AddMenuEntry(
+			FText::FromName(Subject),
+			SubjectText.ToText(),
+			FSlateIcon(),
+			GenerateUIAction(Subject),
+			NAME_None,
+			EUserInterfaceActionType::Check);
+	}
+	MenuBuilder.EndSection();
+	return MenuBuilder.MakeWidget();
+}
+
+void SLiveLinkTimecode::Construct(const FArguments& InArgs)
+{
+	ActiveTimecodeSource = UE::LiveLinkTimecode::Private::System24fps;
+	WorkingClient = (FLiveLinkClient*)&IModularFeatures::Get().GetModularFeature<ILiveLinkClient>(ILiveLinkClient::ModularFeatureName);
+
+	check(WorkingClient);
+	ChildSlot
+	[
+		SNew(SComboButton)
+		.ContentPadding(FMargin(6.0f, 0.0f))
+		.MenuPlacement(MenuPlacement_AboveAnchor)
+		.OnGetMenuContent(this, &SLiveLinkTimecode::MakeMenu)
+		.ComboButtonStyle(&FAppStyle::Get().GetWidgetStyle<FComboButtonStyle>("SimpleComboButton"))
+		.HasDownArrow(true)
+		.ToolTipText(this, &SLiveLinkTimecode::GetTimecodeTooltip)
+		.ButtonContent()
+		[
+			SNew(SHorizontalBox)
+			+ SHorizontalBox::Slot()
+			.AutoWidth()
+			.VAlign(VAlign_Center)
+			.Padding(0, 0, 3, 0)
+			[
+				SNew(STextBlock)
+					.Font(FAppStyle::Get().GetFontStyle("FontAwesome.8"))
+					.ColorAndOpacity(this, &SLiveLinkTimecode::GetTimecodeStatusColor)
+					.Text(FEditorFontGlyphs::Circle)
+			]
+			+ SHorizontalBox::Slot()
+			.AutoWidth()
+			.VAlign(VAlign_Center)
+			.Padding(2, 0, 10, 0)
+			[
+				SNew(STimecode)
+				.DisplayLabel(false)
+				.TimecodeFont(FCoreStyle::Get().GetFontStyle(TEXT("NormalText")))
+				.Timecode(MakeAttributeLambda([]
+				{
+					return FApp::GetTimecode();
+				}))
+			]
+			+ SHorizontalBox::Slot()
+			.AutoWidth()
+			.VAlign(VAlign_Center)
+			.Padding(2, 0, 10, 0)
+			[
+				SNew(STextBlock)
+				.Font(FCoreStyle::Get().GetFontStyle(TEXT("NormalText")))
+				.Text(MakeAttributeLambda([] {return FApp::GetTimecodeFrameRate().ToPrettyText();}))
+			]
+    	]
+	];
+}
+
+#undef LOCTEXT_NAMESPACE
diff --git a/Engine/Plugins/Animation/LiveLinkHub/Source/LiveLinkHub/Private/UI/Widgets/SLiveLinkTimecode.h b/Engine/Plugins/Animation/LiveLinkHub/Source/LiveLinkHub/Private/UI/Widgets/SLiveLinkTimecode.h
new file mode 100644
--- /dev/null
+++ b/Engine/Plugins/Animation/LiveLinkHub/Source/LiveLinkHub/Private/UI/Widgets/SLiveLinkTimecode.h
@@ -0,0 +1,55 @@
+// Copyright Epic Games, Inc. All Rights Reserved.
+#pragma once
+
+#include "Widgets/SCompoundWidget.h"
+#include "UObject/NameTypes.h"
+
+class FName;
+class FLiveLinkClient;
+
+/**
+ * Utility class to display timecode status on the menu bar. Provides user with
+ * a drop down of supported time code values including subjects that are in the
+ * Live Link session.
+ */
+class SLiveLinkTimecode : public SCompoundWidget
+{
+public:
+
+	SLATE_BEGIN_ARGS(SLiveLinkTimecode) {}
+		/** The tab's content */
+	SLATE_END_ARGS()
+
+	/**
+	 * @param InArgs
+	 */
+	void Construct(const FArguments& InArgs);
+
+private:
+	/** Indicate the status of the time code send status */
+	FSlateColor GetTimecodeStatusColor() const;
+
+	/** Gets the tooltip text for the time code widget */
+	FText GetTimecodeTooltip() const;
+
+	/** Makes the combo button menu for Timecode source */
+	TSharedRef<SWidget> MakeMenu();
+
+	/** Invoked by slate when the time code setting is enabled. */
+	void OnEnableTimecodeToggled();
+
+	/** Specifies the source for timecode */
+	void SetTimecodeSource(const FName SourceId);
+
+	/** Sends the current timecode settings to the connected editor(s). */
+	void SendUpdatedTimecodeToEditor();
+
+	/** Do we have timecode source enabled */
+	bool bIsTimecodeSource = false;
+
+	/** Name of the active timecode source. */
+	FName ActiveTimecodeSource;
+
+	/** Cached value for the FLiveLinkClient */
+	FLiveLinkClient* WorkingClient = nullptr;
+};
diff --git a/Engine/Plugins/Animation/LiveLinkHub/Source/LiveLinkHubMessaging/Private/LiveLinkHubMessageBusSource.cpp b/Engine/Plugins/Animation/LiveLinkHub/Source/LiveLinkHubMessaging/Private/LiveLinkHubMessageBusSource.cpp
--- a/Engine/Plugins/Animation/LiveLinkHub/Source/LiveLinkHubMessaging/Private/LiveLinkHubMessageBusSource.cpp
+++ b/Engine/Plugins/Animation/LiveLinkHub/Source/LiveLinkHubMessaging/Private/LiveLinkHubMessageBusSource.cpp
@@ -3,13 +3,17 @@
 #include "LiveLinkHubMessageBusSource.h"
 
 #include "Engine/Level.h"
+#include "Engine/SystemTimeTimecodeProvider.h"
 #include "Engine/World.h"
 #include "HAL/PlatformProcess.h"
 #include "ILiveLinkClient.h"
 #include "ILiveLinkModule.h"
 #include "LiveLinkHubMessages.h"
 #include "LiveLinkMessages.h"
+#include "LiveLinkTimecodeProvider.h"
+#include "LiveLinkTypes.h"
 #include "MessageEndpoint.h"
+#include "MessageEndpointBuilder.h"
 #include "Misc/App.h"
 #include "Misc/EngineVersion.h"
 
@@ -28,7 +32,6 @@ FLiveLinkHubMessageBusSource::FLiveLinkHubMessageBusSource(const FText& InSource
 		LevelEditorModule->OnMapChanged().AddRaw(this, &FLiveLinkHubMessageBusSource::OnMapChanged);
 	}
 #endif
-
 }
 
 FLiveLinkHubMessageBusSource::~FLiveLinkHubMessageBusSource()
@@ -67,6 +70,18 @@ void FLiveLinkHubMessageBusSource::OnMapChanged(UWorld* World, EMapChangeType Ch
 	SendClientInfoMessage();
 }
 
+void FLiveLinkHubMessageBusSource::InitializeMessageEndpoint(FMessageEndpointBuilder& EndpointBuilder)
+{
+	FLiveLinkMessageBusSource::InitializeMessageEndpoint(EndpointBuilder);
+	EndpointBuilder
+		.Handling<FLiveLinkHubTimecodeSettings>(this, &FLiveLinkHubMessageBusSource::HandleTimecodeSettings);
+}
+
+void FLiveLinkHubMessageBusSource::HandleTimecodeSettings(const FLiveLinkHubTimecodeSettings& Message, const TSharedRef<IMessageContext, ESPMode::ThreadSafe>& Context)
+{
+	Message.AssignTimecodeSettingsAsProviderToEngine();
+}
+
 FLiveLinkClientInfoMessage FLiveLinkHubMessageBusSource::CreateLiveLinkClientInfo() const
 {
 	FLiveLinkClientInfoMessage ClientInfo;
diff --git a/Engine/Plugins/Animation/LiveLinkHub/Source/LiveLinkHubMessaging/Private/LiveLinkHubMessageBusSource.h b/Engine/Plugins/Animation/LiveLinkHub/Source/LiveLinkHubMessaging/Private/LiveLinkHubMessageBusSource.h
--- a/Engine/Plugins/Animation/LiveLinkHub/Source/LiveLinkHubMessaging/Private/LiveLinkHubMessageBusSource.h
+++ b/Engine/Plugins/Animation/LiveLinkHub/Source/LiveLinkHubMessaging/Private/LiveLinkHubMessageBusSource.h
@@ -8,6 +8,7 @@ class UWorld;
 enum class EMapChangeType : uint8;
 struct FLiveLinkClientInfoMessage;
 struct FMessageAddress;
+struct FLiveLinkHubTimecodeSettings;
 
 /** LiveLink Message bus source that is connected to a livelink hub. */
 class FLiveLinkHubMessageBusSource : public FLiveLinkMessageBusSource
@@ -18,11 +19,13 @@ public:
 
 protected:
 	//~ Begin FLiveLinkMessageBusSource interface
+	virtual void InitializeMessageEndpoint(FMessageEndpointBuilder& EndpointBuilder);
 	virtual double GetDeadSourceTimeout() const override;
 	virtual void SendConnectMessage() override;
 	//~ End FLiveLinkMessageBusSource
 
 private:
+	void HandleTimecodeSettings(const FLiveLinkHubTimecodeSettings& Message, const TSharedRef<IMessageContext, ESPMode::ThreadSafe>& Context);
 	/** Send information about this UE client to the connected provider */
 	void SendClientInfoMessage();
 	/** Handler called on map changed to update the livelink hub. */
diff --git a/Engine/Plugins/Animation/LiveLinkHub/Source/LiveLinkHubMessaging/Private/LiveLinkHubMessages.cpp b/Engine/Plugins/Animation/LiveLinkHub/Source/LiveLinkHubMessaging/Private/LiveLinkHubMessages.cpp
--- a/Engine/Plugins/Animation/LiveLinkHub/Source/LiveLinkHubMessaging/Private/LiveLinkHubMessages.cpp
+++ b/Engine/Plugins/Animation/LiveLinkHub/Source/LiveLinkHubMessaging/Private/LiveLinkHubMessages.cpp
@@ -2,5 +2,70 @@
 
 #include "LiveLinkHubMessages.h"
 
+#include "Engine/Engine.h"
+#include "Engine/SystemTimeTimecodeProvider.h"
+#include "Engine/World.h"
+#include "Features/IModularFeatures.h"
+#include "ILiveLinkClient.h"
+#include "ILiveLinkModule.h"
+#include "LiveLinkTimecodeProvider.h"
+
 const FName UE::LiveLinkHub::Private::LiveLinkHubProviderType = TEXT("LiveLinkHub");
 FName FLiveLinkHubMessageAnnotation::ProviderTypeAnnotation = TEXT("ProviderType");
+
+
+DEFINE_LOG_CATEGORY_STATIC(LogLiveLinkHubMessages, Log, All);
+
+void FLiveLinkHubTimecodeSettings::AssignTimecodeSettingsAsProviderToEngine() const
+{
+	IModularFeatures& ModularFeatures = IModularFeatures::Get();
+	if (!ModularFeatures.IsModularFeatureAvailable(ILiveLinkClient::ModularFeatureName) || GEngine == nullptr)
+	{
+		return;
+	}
+
+	ILiveLinkClient* LiveLinkClient = &ModularFeatures.GetModularFeature<ILiveLinkClient>(ILiveLinkClient::ModularFeatureName);
+
+	UE_LOG(LogLiveLinkHubMessages, Display, TEXT("Time code change event %s - %s"), *UEnum::GetValueAsName(Source).ToString(), *SubjectName.ToString());
+	if (Source == ELiveLinkHubTimecodeSource::SystemTimeEditor)
+	{
+		// If we are using system time, construct a new system time code provider with the target framerate.
+		FName ObjectName = MakeUniqueObjectName(GEngine, USystemTimeTimecodeProvider::StaticClass(), "DefaultTimecodeProvider");
+		USystemTimeTimecodeProvider* NewTimecodeProvider = NewObject<USystemTimeTimecodeProvider>(GEngine, ObjectName);
+		NewTimecodeProvider->FrameRate = DesiredFrameRate;
+		NewTimecodeProvider->FrameDelay = 0;
+		GEngine->SetTimecodeProvider(NewTimecodeProvider);
+		UE_LOG(LogLiveLinkHubMessages, Display, TEXT("System Time Timecode provider set."));
+	}
+	else if (Source == ELiveLinkHubTimecodeSource::UseSubjectName)
+	{
+		TArray<FLiveLinkSubjectKey> Subjects = LiveLinkClient->GetSubjects(true, true);
+		TOptional<FLiveLinkSubjectKey> Target;
+		// We need to map the named subject to the list of subject keys available.
+		for (const FLiveLinkSubjectKey& Key : Subjects)
+		{
+			if (Key.SubjectName == SubjectName)
+			{
+				Target = Key;
+			}
+		}
+
+		if (Target)
+		{
+			FName ObjectName = MakeUniqueObjectName(GEngine, ULiveLinkTimecodeProvider::StaticClass(), "DefaultLiveLinkTimecodeProvider");
+			ULiveLinkTimecodeProvider* LiveLinkProvider = NewObject<ULiveLinkTimecodeProvider>(GEngine, ObjectName);
+			LiveLinkProvider->SetTargetSubjectKey(*Target);
+			GEngine->SetTimecodeProvider(LiveLinkProvider);
+			UE_LOG(LogLiveLinkHubMessages, Display, TEXT("Live Link Timecode provider assigned to %s."), *SubjectName.ToString());
+		}
+		else
+		{
+			UE_LOG(LogLiveLinkHubMessages, Warning, TEXT("Failed to assign Live Link Timecode provider to %s."), *SubjectName.ToString());
+		}
+	}
+	else
+	{
+		// Force the timecode provider to reset back to the default setting.
+		GEngine->Exec( GEngine->GetCurrentPlayWorld(nullptr), TEXT( "TimecodeProvider.reset" ) );
+	}
+}
diff --git a/Engine/Plugins/Animation/LiveLinkHub/Source/LiveLinkHubMessaging/Public/LiveLinkHubMessages.h b/Engine/Plugins/Animation/LiveLinkHub/Source/LiveLinkHubMessaging/Public/LiveLinkHubMessages.h
--- a/Engine/Plugins/Animation/LiveLinkHub/Source/LiveLinkHubMessaging/Public/LiveLinkHubMessages.h
+++ b/Engine/Plugins/Animation/LiveLinkHub/Source/LiveLinkHubMessaging/Public/LiveLinkHubMessages.h
@@ -4,6 +4,8 @@
 
 #include "LiveLinkMessages.h"
 #include "LiveLinkTypes.h"
+#include "Misc/FrameRate.h"
+
 #include "LiveLinkHubMessages.generated.h"
 
 /** Annotation put on MessageBus messages to indicate the type of provider used. 
@@ -20,6 +22,42 @@ namespace UE::LiveLinkHub::Private
 	extern const LIVELINKHUBMESSAGING_API FName LiveLinkHubProviderType;
 }
 
+
+UENUM()
+enum class ELiveLinkHubTimecodeSource
+{
+	// Not defined by the Hub and thus should use the default system settings.
+	NotDefined,
+
+	// Using system time of the editor.
+	SystemTimeEditor,
+
+	// Using the provided subject name
+	UseSubjectName
+};
+
+/** Special message to communicate / override time code used by the connected editor. */
+USTRUCT()
+struct LIVELINKHUBMESSAGING_API FLiveLinkHubTimecodeSettings
+{
+	GENERATED_BODY()
+
+	/** Source time code value.  If it is not defined then we use the default time code provider in the engine. */
+	UPROPERTY()
+	ELiveLinkHubTimecodeSource Source = ELiveLinkHubTimecodeSource::NotDefined;
+
+	/** Name of the subject to map timecode if Source == ELiveLinkHubTimecodeSource::UseSubjectName */
+	UPROPERTY()
+	FName SubjectName;
+
+	/** Desired frame rate to set if Source == ELiveLinkHubTimecodeSource::SystemTimeEditor. */
+	UPROPERTY()
+	FFrameRate DesiredFrameRate;
+
+	/** Assign the settings to a new timecode provider and override the current engine settings. */
+	void AssignTimecodeSettingsAsProviderToEngine() const;
+};
+
 /** Status of a UE client connected to a live link hub. */
 UENUM()
 enum class ELiveLinkClientStatus
