Fix Optitrack subjects sometimes disappearing or not being transmitted to UE.
This was caused by 2 issues:

1. When removing a subject and quickly sending static data for the same subject, the new subject would get discarded.
2. LLH would always send a RemoveSubject messages a frame later, so if it received a RemoveSubject then a static data message for this subject, livelinkhub would transmit the remove subject message last rather than first, which would cause the new data to be deleted.

Fix subject renaming breaking data transmission

diff --git a/Engine/Plugins/Animation/LiveLink/Source/LiveLink/Private/LiveLinkClient.cpp b/Engine/Plugins/Animation/LiveLink/Source/LiveLink/Private/LiveLinkClient.cpp
--- a/Engine/Plugins/Animation/LiveLink/Source/LiveLink/Private/LiveLinkClient.cpp
+++ b/Engine/Plugins/Animation/LiveLink/Source/LiveLink/Private/LiveLinkClient.cpp
@@ -630,18 +630,21 @@ void FLiveLinkClient::PushSubjectStaticData_Internal(FPendingSubjectStatic&& Sub
 	FLiveLinkSubject* LiveLinkSubject = nullptr;
 	if (FLiveLinkCollectionSubjectItem* SubjectItem = Collection->FindSubject(SubjectStaticData.SubjectKey))
 	{
-		LiveLinkSubject = SubjectItem->GetLiveSubject();
-
-		if (LiveLinkSubject->GetRole() != SubjectStaticData.Role)
+		if (!SubjectItem->bPendingKill)
 		{
-			FLiveLinkLog::Warning(TEXT("Subject '%s' of role '%s' is changing its role to '%s'. Current subject will be removed and a new one will be created"), *SubjectStaticData.SubjectKey.SubjectName.ToString(), *LiveLinkSubject->GetRole().GetDefaultObject()->GetDisplayName().ToString(), *SubjectStaticData.Role.GetDefaultObject()->GetDisplayName().ToString());
+			LiveLinkSubject = SubjectItem->GetLiveSubject();
 
-			Collection->RemoveSubject(SubjectStaticData.SubjectKey);
-			LiveLinkSubject = nullptr;
-		}
-		else
-		{
-			LiveLinkSubject->ClearFrames();
+			if (LiveLinkSubject->GetRole() != SubjectStaticData.Role)
+			{
+				FLiveLinkLog::Warning(TEXT("Subject '%s' of role '%s' is changing its role to '%s'. Current subject will be removed and a new one will be created"), *SubjectStaticData.SubjectKey.SubjectName.ToString(), *LiveLinkSubject->GetRole().GetDefaultObject()->GetDisplayName().ToString(), *SubjectStaticData.Role.GetDefaultObject()->GetDisplayName().ToString());
+
+				Collection->RemoveSubject(SubjectStaticData.SubjectKey);
+				LiveLinkSubject = nullptr;
+			}
+			else
+			{
+				LiveLinkSubject->ClearFrames();
+			}
 		}
 	}
 
diff --git a/Engine/Plugins/Animation/LiveLink/Source/LiveLink/Private/LiveLinkSourceCollection.cpp b/Engine/Plugins/Animation/LiveLink/Source/LiveLink/Private/LiveLinkSourceCollection.cpp
--- a/Engine/Plugins/Animation/LiveLink/Source/LiveLink/Private/LiveLinkSourceCollection.cpp
+++ b/Engine/Plugins/Animation/LiveLink/Source/LiveLink/Private/LiveLinkSourceCollection.cpp
@@ -245,7 +245,7 @@ const FLiveLinkCollectionSubjectItem* FLiveLinkSourceCollection::FindSubject(FLi
 
 const FLiveLinkCollectionSubjectItem* FLiveLinkSourceCollection::FindEnabledSubject(FLiveLinkSubjectName InSubjectName) const
 {
-	return Subjects.FindByPredicate([InSubjectName](const FLiveLinkCollectionSubjectItem& Other) { return Other.Key.SubjectName == InSubjectName && Other.bEnabled; });
+	return Subjects.FindByPredicate([InSubjectName](const FLiveLinkCollectionSubjectItem& Other) { return Other.Key.SubjectName == InSubjectName && Other.bEnabled && !Other.bPendingKill; });
 }
 
 
diff --git a/Engine/Plugins/Animation/LiveLinkHub/Source/LiveLinkHub/Private/LiveLinkHub.cpp b/Engine/Plugins/Animation/LiveLinkHub/Source/LiveLinkHub/Private/LiveLinkHub.cpp
--- a/Engine/Plugins/Animation/LiveLinkHub/Source/LiveLinkHub/Private/LiveLinkHub.cpp
+++ b/Engine/Plugins/Animation/LiveLinkHub/Source/LiveLinkHub/Private/LiveLinkHub.cpp
@@ -60,8 +60,8 @@ void FLiveLinkHub::Initialize()
 
 	LiveLinkHubClient->OnStaticDataReceived_AnyThread().AddSP(this, &FLiveLinkHub::OnStaticDataReceived_AnyThread);
 	LiveLinkHubClient->OnFrameDataReceived_AnyThread().AddSP(this, &FLiveLinkHub::OnFrameDataReceived_AnyThread);
+	LiveLinkHubClient->OnSubjectMarkedPendingKill_AnyThread().AddSP(this, &FLiveLinkHub::OnSubjectMarkedPendingKill_AnyThread);
 	LiveLinkHubClient->OnLiveLinkSubjectAdded().AddSP(this, &FLiveLinkHub::OnSubjectAdded);
-	LiveLinkHubClient->OnLiveLinkSubjectRemoved().AddSP(this, &FLiveLinkHub::OnSubjectRemoved);
 
 	PlaybackController->Start();
 
@@ -73,8 +73,8 @@ FLiveLinkHub::~FLiveLinkHub()
 	RecordingController.Reset();
 	PlaybackController.Reset();
 
-	LiveLinkHubClient->OnLiveLinkSubjectRemoved().RemoveAll(this);
 	LiveLinkHubClient->OnLiveLinkSubjectAdded().RemoveAll(this);
+	LiveLinkHubClient->OnSubjectMarkedPendingKill_AnyThread().RemoveAll(this);
 	LiveLinkHubClient->OnFrameDataReceived_AnyThread().RemoveAll(this);
 	LiveLinkHubClient->OnStaticDataReceived_AnyThread().RemoveAll(this);
 
@@ -140,6 +140,8 @@ void FLiveLinkHub::OnStaticDataReceived_AnyThread(const FLiveLinkSubjectKey& InS
 	FLiveLinkStaticDataStruct StaticDataCopy;
 	StaticDataCopy.InitializeWith(InStaticDataStruct);
 
+	UE_LOG(LogLiveLinkHub, Verbose, TEXT("Pushing static data for %s"), *InSubjectKey.SubjectName.ToString());
+
 	const FName OverridenName = GetSubjectNameOverride(InSubjectKey);
 	LiveLinkProvider->UpdateSubjectStaticData(OverridenName, InRole, MoveTemp(StaticDataCopy));
 }
@@ -156,7 +158,7 @@ void FLiveLinkHub::OnFrameDataReceived_AnyThread(const FLiveLinkSubjectKey& InSu
 	FrameDataCopy.InitializeWith(InFrameDataStruct);
 
 	const FName OverridenName = GetSubjectNameOverride(InSubjectKey);
-	if (LiveLinkHubClient->IsSubjectEnabled(OverridenName))
+	if (LiveLinkHubClient->IsSubjectEnabled(InSubjectKey.SubjectName))
 	{
 		LiveLinkProvider->UpdateSubjectFrameData(OverridenName, MoveTemp(FrameDataCopy));
 	}
@@ -173,14 +175,25 @@ void FLiveLinkHub::OnSubjectAdded(FLiveLinkSubjectKey InSubjectKey) const
 		StaticDataCopy.InitializeWith(*StaticData);
 
 		const FName OverridenName = GetSubjectNameOverride(InSubjectKey);
+
+		UE_LOG(LogLiveLinkHub, Verbose, TEXT("Pushed static subject %s"), *InSubjectKey.SubjectName.ToString());
 		LiveLinkProvider->UpdateSubjectStaticData(OverridenName, SubjectSettings->Role, MoveTemp(StaticDataCopy));
 	}
+	else
+	{
+		UE_LOG(LogLiveLinkHub, Verbose, TEXT("Failed to push static subject %s, static data doesn't exist"), *InSubjectKey.SubjectName.ToString());
+	}
 }
 
-void FLiveLinkHub::OnSubjectRemoved(FLiveLinkSubjectKey InSubjectKey) const
+void FLiveLinkHub::OnSubjectMarkedPendingKill_AnyThread(const FLiveLinkSubjectKey& InSubjectKey) const
 {
+	UE_LOG(LogLiveLinkHub, Verbose, TEXT("Removed subject %s"), *InSubjectKey.SubjectName.ToString());
+
 	// Send an update to connected clients as well.
 	const FName OverridenName = GetSubjectNameOverride(InSubjectKey);
+
+	// Note: We send a RemoveSubject message to connected clients when the subject is marked pending kill in order to process this message in the right order.
+	// If we were to send a RemoveSubject message after the OnSubjectRemoved delegate, it could cause our RemoveSubject message to be sent out of order.
 	LiveLinkProvider->RemoveSubject(OverridenName);
 }
 
diff --git a/Engine/Plugins/Animation/LiveLinkHub/Source/LiveLinkHub/Private/LiveLinkHub.h b/Engine/Plugins/Animation/LiveLinkHub/Source/LiveLinkHub/Private/LiveLinkHub.h
--- a/Engine/Plugins/Animation/LiveLinkHub/Source/LiveLinkHub/Private/LiveLinkHub.h
+++ b/Engine/Plugins/Animation/LiveLinkHub/Source/LiveLinkHub/Private/LiveLinkHub.h
@@ -74,8 +74,8 @@ private:
 	//~ LiveLink Client delegates
 	void OnStaticDataReceived_AnyThread(const FLiveLinkSubjectKey& InSubjectKey, TSubclassOf<ULiveLinkRole> InRole, const FLiveLinkStaticDataStruct& InStaticDataStruct) const;
 	void OnFrameDataReceived_AnyThread(const FLiveLinkSubjectKey& InSubjectKey, const FLiveLinkFrameDataStruct& InFrameDataStruct) const;
+	void OnSubjectMarkedPendingKill_AnyThread(const FLiveLinkSubjectKey& InSubjectKey) const;
 	void OnSubjectAdded(FLiveLinkSubjectKey InSubjectKey) const;
-	void OnSubjectRemoved(FLiveLinkSubjectKey InSubjectKey) const;
 	//~ LiveLink Client delegates
 
 	/** Bind all available live link hub commands. */
diff --git a/Engine/Plugins/Animation/LiveLinkHub/Source/LiveLinkHub/Private/LiveLinkHubClient.cpp b/Engine/Plugins/Animation/LiveLinkHub/Source/LiveLinkHub/Private/LiveLinkHubClient.cpp
--- a/Engine/Plugins/Animation/LiveLinkHub/Source/LiveLinkHub/Private/LiveLinkHubClient.cpp
+++ b/Engine/Plugins/Animation/LiveLinkHub/Source/LiveLinkHub/Private/LiveLinkHubClient.cpp
@@ -281,18 +281,21 @@ void FLiveLinkHubClient::PushSubjectStaticData_AnyThread(const FLiveLinkSubjectK
 		FScopeLock Lock(&CollectionAccessCriticalSection);
 		if (FLiveLinkCollectionSubjectItem* SubjectItem = Collection->FindSubject(SubjectKey))
 		{
-			LiveLinkSubject = SubjectItem->GetLiveSubject();
-
-			if (LiveLinkSubject->GetRole() != Role)
+			if (!SubjectItem->bPendingKill)
 			{
-				UE_LOG(LogLiveLinkHub, Warning, TEXT("Subject '%s' of role '%s' is changing its role to '%s'. Current subject will be removed and a new one will be created"), *SubjectKey.SubjectName.ToString(), *LiveLinkSubject->GetRole().GetDefaultObject()->GetDisplayName().ToString(), *Role.GetDefaultObject()->GetDisplayName().ToString());
+				LiveLinkSubject = SubjectItem->GetLiveSubject();
 
-				Collection->RemoveSubject(SubjectKey);
-				LiveLinkSubject = nullptr;
-			}
-			else
-			{
-				LiveLinkSubject->ClearFrames();
+				if (LiveLinkSubject->GetRole() != Role)
+				{
+					UE_LOG(LogLiveLinkHub, Warning, TEXT("Subject '%s' of role '%s' is changing its role to '%s'. Current subject will be removed and a new one will be created"), *SubjectKey.SubjectName.ToString(), *LiveLinkSubject->GetRole().GetDefaultObject()->GetDisplayName().ToString(), *Role.GetDefaultObject()->GetDisplayName().ToString());
+
+					Collection->RemoveSubject(SubjectKey);
+					LiveLinkSubject = nullptr;
+				}
+				else
+				{
+					LiveLinkSubject->ClearFrames();
+				}
 			}
 		}
 
@@ -426,4 +429,10 @@ bool FLiveLinkHubClient::IsSubjectValid(const FLiveLinkSubjectKey& InSubjectKey)
 	return false;
 }
 
+void FLiveLinkHubClient::RemoveSubject_AnyThread(const FLiveLinkSubjectKey& InSubjectKey)
+{
+	OnSubjectMarkedPendingKill_AnyThread().Broadcast(InSubjectKey);
+	FLiveLinkClient::RemoveSubject_AnyThread(InSubjectKey);
+}
+
 #undef LOCTEXT_NAMESPACE
diff --git a/Engine/Plugins/Animation/LiveLinkHub/Source/LiveLinkHub/Private/LiveLinkHubClient.h b/Engine/Plugins/Animation/LiveLinkHub/Source/LiveLinkHub/Private/LiveLinkHubClient.h
--- a/Engine/Plugins/Animation/LiveLinkHub/Source/LiveLinkHub/Private/LiveLinkHubClient.h
+++ b/Engine/Plugins/Animation/LiveLinkHub/Source/LiveLinkHub/Private/LiveLinkHubClient.h
@@ -15,6 +15,8 @@ struct ILiveLinkProvider;
 
 DECLARE_TS_MULTICAST_DELEGATE_TwoParams(FOnFrameDataReceived_AnyThread, const FLiveLinkSubjectKey&, const FLiveLinkFrameDataStruct&);
 DECLARE_TS_MULTICAST_DELEGATE_ThreeParams(FOnStaticDataReceived_AnyThread, const FLiveLinkSubjectKey&, TSubclassOf<ULiveLinkRole>, const FLiveLinkStaticDataStruct&);
+DECLARE_TS_MULTICAST_DELEGATE_OneParam(FOnSubjectMarkedPendingKill_AnyThread, const FLiveLinkSubjectKey&);
+
 
 class FLiveLinkHubClient : public FLiveLinkClient
 {
@@ -37,6 +39,15 @@ public:
 		return OnStaticDataReceivedDelegate_AnyThread;
 	}
 
+	/** 
+	 * Get the delegate called when a subject is marked for deletion. 
+	 * This delegate will fire as soon as the subject is marked for deletion, while the OnSubjectRemovedDelegate may trigger at a later time.
+	 */
+	FOnSubjectMarkedPendingKill_AnyThread& OnSubjectMarkedPendingKill_AnyThread()
+	{
+		return OnSubjectMarkedPendingKillDelegate_AnyThread;
+	}
+
 public:
 	//~ Begin ILiveLinkClient interface
 	virtual bool CreateSource(const FLiveLinkSourcePreset& InSourcePreset) override;
@@ -45,6 +56,7 @@ public:
 	virtual void PushSubjectFrameData_AnyThread(const FLiveLinkSubjectKey& SubjectKey, FLiveLinkFrameDataStruct&& FrameData) override;
 	virtual FText GetSourceStatus(FGuid InEntryGuid) const override;
 	virtual bool IsSubjectValid(const FLiveLinkSubjectKey& InSubjectKey) const override;
+	virtual void RemoveSubject_AnyThread(const FLiveLinkSubjectKey& InSubjectKey) override;
 	//~ End ILiveLinkClient interface
 
 private:
@@ -62,4 +74,6 @@ private:
     FOnFrameDataReceived_AnyThread OnFrameDataReceivedDelegate_AnyThread;
     /** Delegate called when static data is received. */
     FOnStaticDataReceived_AnyThread OnStaticDataReceivedDelegate_AnyThread;
+	/** Delegate called when a subject is marked for deletion. */
+	FOnSubjectMarkedPendingKill_AnyThread OnSubjectMarkedPendingKillDelegate_AnyThread;
 };
