LiveLinkHub - Fix subjects always appearing as valid

Update the last frame receive time when we intercept a livelink message, and query it when we check for subject validity.
The timeout duration can be adjusted by modifying LiveLinkSettings' TimeWithoutFrameToBeConsiderAsInvalid

Fix issue with playback widget not being anchored to the bottom of the hub

diff --git a/Engine/Plugins/Animation/LiveLink/Source/LiveLink/Public/LiveLinkSubject.h b/Engine/Plugins/Animation/LiveLink/Source/LiveLink/Public/LiveLinkSubject.h
--- a/Engine/Plugins/Animation/LiveLink/Source/LiveLink/Public/LiveLinkSubject.h
+++ b/Engine/Plugins/Animation/LiveLink/Source/LiveLink/Public/LiveLinkSubject.h
@@ -106,7 +106,10 @@ public:
 	FLiveLinkSubjectTimeSyncData GetTimeSyncData();
 	bool IsTimeSynchronized() const;
 
+	/** Get the timestamp of the last time a frame was received for this subject. */
 	double GetLastPushTime() const { return LastPushTime; }
+	/** Set the last time a frame was received. */
+	void SetLastPushTime(double InLastPushTime) { LastPushTime = InLastPushTime; }
 
 private:
 	int32 FindNewFrame_WorldTime(const FLiveLinkWorldTime& FrameTime) const;
diff --git a/Engine/Plugins/Animation/LiveLinkHub/Source/LiveLinkHub/Private/LiveLinkHubClient.cpp b/Engine/Plugins/Animation/LiveLinkHub/Source/LiveLinkHub/Private/LiveLinkHubClient.cpp
--- a/Engine/Plugins/Animation/LiveLinkHub/Source/LiveLinkHub/Private/LiveLinkHubClient.cpp
+++ b/Engine/Plugins/Animation/LiveLinkHub/Source/LiveLinkHub/Private/LiveLinkHubClient.cpp
@@ -13,6 +13,7 @@
 #include "LiveLinkSubject.h"
 #include "LiveLinkSubjectSettings.h"
 #include "LiveLinkTimedDataInput.h"
+#include "Misc/App.h"
 #include "Modules/ModuleManager.h"
 #include "Recording/LiveLinkPlaybackSource.h"
 #include "Recording/LiveLinkPlaybackSubject.h"
@@ -405,6 +406,14 @@ void FLiveLinkHubClient::PushSubjectFrameData_AnyThread(const FLiveLinkSubjectKe
 {
 	SCOPE_CYCLE_COUNTER(STAT_LiveLinkHub_PushFrameData);
 
+	if (const FLiveLinkCollectionSubjectItem* SubjectItem = Collection->FindSubject(SubjectKey))
+	{
+		if (FLiveLinkSubject* LiveSubject = SubjectItem->GetLiveSubject())
+		{
+			LiveSubject->SetLastPushTime(FApp::GetCurrentTime());
+		}
+	}
+
 	OnFrameDataReceivedDelegate_AnyThread.Broadcast(SubjectKey, FrameData);
 }
 
@@ -423,7 +432,11 @@ bool FLiveLinkHubClient::IsSubjectValid(const FLiveLinkSubjectKey& InSubjectKey)
 	FScopeLock Lock(&CollectionAccessCriticalSection);
 	if (const FLiveLinkCollectionSubjectItem* SubjectItem = Collection->FindSubject(InSubjectKey))
 	{
-		// todo LiveLinkHub : Get subject validity from the UE client since we don't store snapshots on the hub side yet
+		if (FLiveLinkSubject* LiveSubject = SubjectItem->GetLiveSubject())
+		{
+			// We don't store frame snapshots so we instead rely on the subject's last push time.
+			return FApp::GetCurrentTime() - LiveSubject->GetLastPushTime() < GetDefault<ULiveLinkSettings>()->GetTimeWithoutFrameToBeConsiderAsInvalid();
+		}
 		return true;
 	}
 	return false;
diff --git a/Engine/Plugins/Animation/LiveLinkHub/Source/LiveLinkHub/Private/Recording/SLiveLinkHubRecordingListView.h b/Engine/Plugins/Animation/LiveLinkHub/Source/LiveLinkHub/Private/Recording/SLiveLinkHubRecordingListView.h
--- a/Engine/Plugins/Animation/LiveLinkHub/Source/LiveLinkHub/Private/Recording/SLiveLinkHubRecordingListView.h
+++ b/Engine/Plugins/Animation/LiveLinkHub/Source/LiveLinkHub/Private/Recording/SLiveLinkHubRecordingListView.h
@@ -67,6 +67,7 @@ public:
 				]
 			]
 			+ SVerticalBox::Slot()
+			.FillHeight(1.0f)
 			[
 				CreateRecordingPicker()
 			]
@@ -126,11 +127,11 @@ private:
 		MenuBuilder.BeginSection(NAME_None, LOCTEXT("ImportRecording_MenuSection", "Import Recording"));
 		{
 			TSharedRef<SWidget> PresetPicker = SNew(SBox)
-				.MinDesiredWidth(400.f)
-				.MinDesiredHeight(400.f)
-				[
-					ContentBrowser.CreateAssetPicker(AssetPickerConfig)
-				];
+			.VAlign(VAlign_Fill)
+			.HAlign(HAlign_Fill)
+			[
+				ContentBrowser.CreateAssetPicker(AssetPickerConfig)
+			];
 
 			MenuBuilder.AddWidget(PresetPicker, FText(), true, false);
 		}
diff --git a/Engine/Plugins/Animation/LiveLinkHub/Source/LiveLinkHub/Private/UI/Widgets/SLiveLinkHubMainTabView.cpp b/Engine/Plugins/Animation/LiveLinkHub/Source/LiveLinkHub/Private/UI/Widgets/SLiveLinkHubMainTabView.cpp
--- a/Engine/Plugins/Animation/LiveLinkHub/Source/LiveLinkHub/Private/UI/Widgets/SLiveLinkHubMainTabView.cpp
+++ b/Engine/Plugins/Animation/LiveLinkHub/Source/LiveLinkHub/Private/UI/Widgets/SLiveLinkHubMainTabView.cpp
@@ -239,7 +239,7 @@ TSharedRef<SDockTab> SLiveLinkHubMainTabView::SpawnPlaybackTab(const FSpawnTabAr
 		[
 			SNew(SVerticalBox)
 			+ SVerticalBox::Slot()
-			.AutoHeight()
+			.FillHeight(1.0f)
 			[
 				LiveLinkHubModule.GetRecordingListController()->MakeRecordingList()
 			]
