LiveLinkHub: Fix and improve frametime / seconds / timecode conversion.
Frame time can now be properly displayed and converted between seconds, frames, and timecode.

diff --git a/Engine/Plugins/Animation/LiveLinkHub/Source/LiveLinkHub/Private/Recording/Implementations/LiveLinkUAssetRecordingPlayer.cpp b/Engine/Plugins/Animation/LiveLinkHub/Source/LiveLinkHub/Private/Recording/Implementations/LiveLinkUAssetRecordingPlayer.cpp
--- a/Engine/Plugins/Animation/LiveLinkHub/Source/LiveLinkHub/Private/Recording/Implementations/LiveLinkUAssetRecordingPlayer.cpp
+++ b/Engine/Plugins/Animation/LiveLinkHub/Source/LiveLinkHub/Private/Recording/Implementations/LiveLinkUAssetRecordingPlayer.cpp
@@ -2,6 +2,7 @@
 
 #include "LiveLinkUAssetRecordingPlayer.h"
 
+#include "LiveLinkHubLog.h"
 #include "Recording/LiveLinkRecording.h"
 
 
@@ -286,6 +287,24 @@ void FLiveLinkPlaybackTracks::Restart(int32 InIndex)
 	}
 }
 
+FFrameRate FLiveLinkPlaybackTracks::GetInitialFrameRate() const
+{
+	for (const FLiveLinkPlaybackTrack& Track : Tracks)
+	{
+		if (Track.LiveLinkRole == nullptr && Track.FrameData.Num() > 0)
+		{
+			FLiveLinkFrameDataStruct FrameDataStruct;
+			FrameDataStruct.InitializeWith(Track.FrameData[0].GetScriptStruct(), (FLiveLinkBaseFrameData*)Track.FrameData[0].GetMemory());
+
+			return FrameDataStruct.GetBaseData()->MetaData.SceneTime.Rate;
+		}
+	}
+
+	UE_LOG(LogLiveLinkHub, Warning, TEXT("Could not find an initial framerate for the recording. Using the defualt value."));
+	
+	return FFrameRate(30, 1);
+}
+
 void FLiveLinkUAssetRecordingPlayer::PreparePlayback(const ULiveLinkRecording* CurrentRecording)
 {
 	const ULiveLinkUAssetRecording* UAssetRecording = CastChecked<ULiveLinkUAssetRecording>(CurrentRecording);
diff --git a/Engine/Plugins/Animation/LiveLinkHub/Source/LiveLinkHub/Private/Recording/Implementations/LiveLinkUAssetRecordingPlayer.h b/Engine/Plugins/Animation/LiveLinkHub/Source/LiveLinkHub/Private/Recording/Implementations/LiveLinkUAssetRecordingPlayer.h
--- a/Engine/Plugins/Animation/LiveLinkHub/Source/LiveLinkHub/Private/Recording/Implementations/LiveLinkUAssetRecordingPlayer.h
+++ b/Engine/Plugins/Animation/LiveLinkHub/Source/LiveLinkHub/Private/Recording/Implementations/LiveLinkUAssetRecordingPlayer.h
@@ -71,6 +71,9 @@ struct FLiveLinkPlaybackTracks
 	
 	void Restart(int32 InIndex);
 
+	/** Retrieve the framerate of the first frame */
+	FFrameRate GetInitialFrameRate() const;
+
 public:
 	/** LiveLink tracks to playback. */
 	TArray<FLiveLinkPlaybackTrack> Tracks;
@@ -110,6 +113,11 @@ public:
 	{
 		CurrentRecordingPlayback.Restart(InIndex);
 	}
+
+	virtual FFrameRate GetInitialFramerate() override
+	{
+		return CurrentRecordingPlayback.GetInitialFrameRate();
+	}
 	
 private:
 	/** All tracks for the current recording. */
diff --git a/Engine/Plugins/Animation/LiveLinkHub/Source/LiveLinkHub/Private/Recording/LiveLinkHubPlaybackController.cpp b/Engine/Plugins/Animation/LiveLinkHub/Source/LiveLinkHub/Private/Recording/LiveLinkHubPlaybackController.cpp
--- a/Engine/Plugins/Animation/LiveLinkHub/Source/LiveLinkHub/Private/Recording/LiveLinkHubPlaybackController.cpp
+++ b/Engine/Plugins/Animation/LiveLinkHub/Source/LiveLinkHub/Private/Recording/LiveLinkHubPlaybackController.cpp
@@ -17,11 +17,53 @@
 #include "UObject/Object.h"
 #include "UObject/Package.h"
 
+/**
+ * Thread safe way to load and store a FQualifiedFrameTime. This is necessary because atomic<FQualifiedFrameTime> isn't
+ * necessarily cross platform compatible when performing copy operations.
+ */
+class FLiveLinkHubAtomicQualifiedFrameTime {
+public:
+	FLiveLinkHubAtomicQualifiedFrameTime()
+		: Value(FQualifiedFrameTime()) {}
+
+	FLiveLinkHubAtomicQualifiedFrameTime(const FLiveLinkHubAtomicQualifiedFrameTime& Other) {
+		FScopeLock Lock(&Other.Mutex);
+		Value = Other.Value;
+	}
+
+	FLiveLinkHubAtomicQualifiedFrameTime& operator=(const FLiveLinkHubAtomicQualifiedFrameTime& Other) {
+		if (this != &Other) {
+			FScopeLock Lock(&Other.Mutex);
+			Value = Other.Value;
+		}
+		return *this;
+	}
+
+	/** Set the underlying value. */
+	void SetValue(const FQualifiedFrameTime& NewPlayhead) {
+		FScopeLock Lock(&Mutex);
+		Value = NewPlayhead;
+	}
+
+	/** Retrieve the underlying value. */
+	FQualifiedFrameTime GetValue() const {
+		FScopeLock Lock(&Mutex);
+		return Value;
+	}
+
+private:
+	/** Mutex for reading/writing underlying value. */
+	mutable FCriticalSection Mutex;
+	/** The underlying qualified frame time value. */
+	FQualifiedFrameTime Value;
+};
+
 FLiveLinkHubPlaybackController::FLiveLinkHubPlaybackController()
 {
 	Client = &IModularFeatures::Get().GetModularFeature<ILiveLinkClient>(ILiveLinkClient::ModularFeatureName);
 	
 	RecordingPlayer = MakeUnique<FLiveLinkUAssetRecordingPlayer>();
+	Playhead = MakeShared<FLiveLinkHubAtomicQualifiedFrameTime, ESPMode::ThreadSafe>();
 }
 
 FLiveLinkHubPlaybackController::~FLiveLinkHubPlaybackController()
@@ -45,19 +87,19 @@ TSharedRef<SWidget> FLiveLinkHubPlaybackController::MakePlaybackWidget()
 		.OnPlayReverse_Raw(this, &FLiveLinkHubPlaybackController::BeginPlayback, true)
 		.OnFirstFrame_Lambda([this]()
 		{
-			GoToFrame(GetSelectionStartFrame());
+			GoToTime(GetSelectionStartTime());
 		})
 		.OnLastFrame_Lambda([this]()
 		{
-			GoToFrame(GetSelectionEndFrame());
+			GoToTime(GetSelectionEndTime());
 		})
 		.OnPreviousFrame_Lambda([this]()
 		{
-			GoToFrame(CurrentFrameIndex - 1);
+			GoToTime(FQualifiedFrameTime(FFrameTime(GetCurrentFrame() - 1), GetFrameRate()));
 		})
 		.OnNextFrame_Lambda([this]()
 		{
-			GoToFrame(CurrentFrameIndex + 1);
+			GoToTime(FQualifiedFrameTime(FFrameTime(GetCurrentFrame() + 1), GetFrameRate()));
 		})
 		.SetCurrentTime_Raw(this, &FLiveLinkHubPlaybackController::GoToTime)
 		.GetViewRange_Lambda([this]()
@@ -70,7 +112,6 @@ TSharedRef<SWidget> FLiveLinkHubPlaybackController::MakePlaybackWidget()
 		})
 		.GetTotalLength_Raw(this, &FLiveLinkHubPlaybackController::GetLength)
 		.GetCurrentTime_Raw(this, &FLiveLinkHubPlaybackController::GetCurrentTime)
-		.GetCurrentFrame_Raw(this, &FLiveLinkHubPlaybackController::GetCurrentFrame)
 		.GetSelectionStartTime_Raw(this, &FLiveLinkHubPlaybackController::GetSelectionStartTime)
 		.SetSelectionStartTime_Raw(this,& FLiveLinkHubPlaybackController::SetSelectionStartTime)
 		.GetSelectionEndTime_Raw(this, &FLiveLinkHubPlaybackController::GetSelectionEndTime)
@@ -79,7 +120,7 @@ TSharedRef<SWidget> FLiveLinkHubPlaybackController::MakePlaybackWidget()
 		.IsInReverse_Raw(this, &FLiveLinkHubPlaybackController::IsPlayingInReverse)
 		.IsLooping_Raw(this, &FLiveLinkHubPlaybackController::IsLooping)
 		.OnSetLooping_Raw(this, &FLiveLinkHubPlaybackController::SetLooping)
-		.GetTimeDelta_Raw(this, &FLiveLinkHubPlaybackController::GetTimeDelta);
+		.GetFrameRate_Raw(this, &FLiveLinkHubPlaybackController::GetFrameRate);
 }
 
 void FLiveLinkHubPlaybackController::StartPlayback()
@@ -95,9 +136,23 @@ void FLiveLinkHubPlaybackController::ResumePlayback()
 {
 	bIsInPlayback = true;
 	bIsPaused = false;
-	StartTimestamp = GetTimeFromFrameIndex(GetCurrentFrame());
+	FQualifiedFrameTime CurrentTime = GetCurrentTime();
+
+	// Clamp to selection start/end.
+	if (CurrentTime.AsSeconds() < GetSelectionStartTime().AsSeconds() && !bIsReverse)
+	{
+		CurrentTime = GetSelectionStartTime();
+	}
+	else if (CurrentTime.AsSeconds() > GetSelectionEndTime().AsSeconds() && bIsReverse)
+	{
+		CurrentTime = GetSelectionEndTime();
+	}
+
+	Playhead->SetValue(CurrentTime);
+	
+	StartTimestamp = CurrentTime.AsSeconds();
 	// Force sync so interpolation doesn't interfere if the first frame isn't the current frame
-	SyncToFrame(CurrentFrameIndex);
+	SyncToFrame(CurrentTime.Time.GetFrame());
 }
 
 void FLiveLinkHubPlaybackController::PreparePlayback(ULiveLinkRecording* InLiveLinkRecording)
@@ -114,20 +169,21 @@ void FLiveLinkHubPlaybackController::PreparePlayback(ULiveLinkRecording* InLiveL
 		}
 		
 		RecordingToPlay.Reset(InLiveLinkRecording);
+		RecordingPlayer->PreparePlayback(RecordingToPlay.Get());
+		
+		CurrentFrameRate = RecordingPlayer->GetInitialFramerate();
 
 		// The start and end of playback.
-		SelectionStartTime = 0.f;
-		SelectionEndTime = GetLength();
+		SetSelectionStartTime(FQualifiedFrameTime(FFrameTime::FromDecimal(0.f), GetFrameRate()));
+		SetSelectionEndTime(GetLength());
 		
 		// The range the user sees.
-		SliderViewRange = TRange<double>(SelectionStartTime, SelectionEndTime);
+		SliderViewRange = TRange<double>(SelectionStartTime.AsSeconds(), SelectionEndTime.AsSeconds());
 		
 		RollbackPreset.Reset(NewObject<ULiveLinkPreset>(GetTransientPackage(), TEXT("RecordingRollbackPreset")));
 		// Save the current state of the sources/subjects in a rollback preset.
 		RollbackPreset->BuildFromClient();
 
-		RecordingPlayer->PreparePlayback(RecordingToPlay.Get());
-
 		RecordingToPlay->RecordingPreset->ApplyToClientLatent([this](bool)
 		{
 			bIsReady = true;
@@ -158,7 +214,7 @@ void FLiveLinkHubPlaybackController::BeginPlayback(bool bInReverse)
 		{
 			if (bIsReverse)
 			{
-				RecordingPlayer->RestartPlayback(CurrentFrameIndex);
+				RecordingPlayer->RestartPlayback(GetCurrentFrame().Value);
 			}
 			
 			// Resume as normal for anywhere else in the recording.
@@ -178,12 +234,10 @@ void FLiveLinkHubPlaybackController::BeginPlayback(bool bInReverse)
 
 void FLiveLinkHubPlaybackController::RestartPlayback()
 {
-	CurrentFrameIndex = INDEX_NONE;
-
 	const bool bOldReverse = bIsReverse; // Stop playback resets reverse
 	StopPlayback();
-	StartTimestamp = GetTimeFromFrameIndex(CurrentFrameIndex);
-	RecordingPlayer->RestartPlayback(CurrentFrameIndex);
+	StartTimestamp = GetCurrentTime().AsSeconds();
+	RecordingPlayer->RestartPlayback(GetCurrentFrame().Value);
 	bIsInPlayback = true;
 	bIsReverse = bOldReverse;
 }
@@ -205,11 +259,8 @@ void FLiveLinkHubPlaybackController::StopPlayback()
 	}
 	
 	const bool bReverse = bIsReverse.load();
-	Playhead = bReverse ? GetSelectionEndTime() : GetSelectionStartTime();
-	if (CurrentFrameIndex == INDEX_NONE)
-	{
-		CurrentFrameIndex = GetFrameIndexFromTime(bReverse ? GetSelectionEndTime() : GetSelectionStartTime());
-	}
+	Playhead->SetValue(bReverse ? GetSelectionEndTime() : GetSelectionStartTime());
+
 	PlaybackStartTime = FPlatformTime::Seconds();
 
 	RecordingPlayer->RestartPlayback();
@@ -223,12 +274,11 @@ void FLiveLinkHubPlaybackController::Eject()
 	StopPlayback();
 	
 	bIsPaused = false;
-	CurrentFrameIndex = 0;
-	RecordingPlayer->RestartPlayback(CurrentFrameIndex);
+	RecordingPlayer->RestartPlayback(0);
 
-	SetSelectionStartTime(0.f);
-	SetSelectionEndTime(0.f);
-	Playhead = 0.f;
+	SetSelectionStartTime(FQualifiedFrameTime(FFrameTime::FromDecimal(0), GetFrameRate()));
+	SetSelectionEndTime(FQualifiedFrameTime(FFrameTime::FromDecimal(0), GetFrameRate()));
+	Playhead->SetValue(FQualifiedFrameTime(FFrameTime::FromDecimal(0), GetFrameRate()));
 	StartTimestamp = 0.f;
 	
 	if (RollbackPreset.IsValid())
@@ -239,83 +289,69 @@ void FLiveLinkHubPlaybackController::Eject()
 	RecordingToPlay.Reset();
 }
 
-void FLiveLinkHubPlaybackController::GoToTime(double InTime)
-{
-	// Stop needs to occur to restart playback.
-	StopPlayback();
-
-	PlaybackStartTime -= InTime;
-	Playhead = InTime;
-
-	const int32 FrameIndex = GetFrameIndexFromTime(InTime);
-	SyncToFrame(FrameIndex);
-}
-
-void FLiveLinkHubPlaybackController::GoToFrame(int32 InFrameIndex)
+void FLiveLinkHubPlaybackController::GoToTime(FQualifiedFrameTime InTime)
 {
 	// Stop needs to occur to restart playback.
 	StopPlayback();
 
-	const double NewTime = GetTimeFromFrameIndex(InFrameIndex);
-	PlaybackStartTime -= NewTime;
-	Playhead = NewTime;
+	const double TimeDouble = InTime.AsSeconds();
 	
-	SyncToFrame(InFrameIndex);
-}
+	PlaybackStartTime -= TimeDouble;
+	Playhead->SetValue(InTime);
 
-int32 FLiveLinkHubPlaybackController::GetFrameIndexFromTime(double InTime, bool bReverse) const
-{
-	return RecordingPlayer->PlayheadToFrameIndex(InTime, bReverse);
+	SyncToFrame(InTime.Time.GetFrame());
 }
 
-double FLiveLinkHubPlaybackController::GetTimeFromFrameIndex(int32 InFrameIndex) const
+FQualifiedFrameTime FLiveLinkHubPlaybackController::GetSelectionStartTime() const
 {
-	return RecordingPlayer->FrameIndexToPlayhead(InFrameIndex);
+	return SelectionStartTime;
 }
 
-int32 FLiveLinkHubPlaybackController::GetSelectionStartFrame() const
+void FLiveLinkHubPlaybackController::SetSelectionStartTime(FQualifiedFrameTime InTime)
 {
-	return GetFrameIndexFromTime(GetSelectionStartTime(), bIsReverse);
+	SelectionStartTime = InTime;
 }
 
-int32 FLiveLinkHubPlaybackController::GetSelectionEndFrame() const
+FQualifiedFrameTime FLiveLinkHubPlaybackController::GetSelectionEndTime() const
 {
-	return GetFrameIndexFromTime(GetSelectionEndTime(), bIsReverse);
+	return SelectionEndTime;
 }
 
-double FLiveLinkHubPlaybackController::GetSelectionStartTime() const
+void FLiveLinkHubPlaybackController::SetSelectionEndTime(FQualifiedFrameTime InTime)
 {
-	return SelectionStartTime;
+	SelectionEndTime = InTime;
 }
 
-void FLiveLinkHubPlaybackController::SetSelectionStartTime(double InTime)
+FQualifiedFrameTime FLiveLinkHubPlaybackController::GetLength() const
 {
-	SelectionStartTime = InTime;
-}
+	const FFrameRate FrameRate = GetFrameRate();
 
-double FLiveLinkHubPlaybackController::GetSelectionEndTime() const
-{
-	return SelectionEndTime;
-}
+	const double Length = RecordingToPlay ? RecordingToPlay->LengthInSeconds : 0.f;
+	
+	const double TotalFramesDouble = Length * FrameRate.Numerator;
+	const int32 TotalFrames = FMath::FloorToInt(TotalFramesDouble);
+	
+	const FFrameNumber LastFrameNumber(TotalFrames - 1);
+	
+	const FQualifiedFrameTime FrameTime(LastFrameNumber, FrameRate);
 
-void FLiveLinkHubPlaybackController::SetSelectionEndTime(double InTime)
-{
-	SelectionEndTime = InTime;
+	return FrameTime;
 }
 
-double FLiveLinkHubPlaybackController::GetLength() const
+FQualifiedFrameTime FLiveLinkHubPlaybackController::GetCurrentTime() const
 {
-	return RecordingToPlay ? RecordingToPlay->LengthInSeconds : 0.f;
+	const FQualifiedFrameTime Time = Playhead->GetValue();
+	return Time;
 }
 
-double FLiveLinkHubPlaybackController::GetCurrentTime() const
+FFrameNumber FLiveLinkHubPlaybackController::GetCurrentFrame() const
 {
-	return Playhead.load();
+	return GetCurrentTime().Time.GetFrame();
 }
 
-int32 FLiveLinkHubPlaybackController::GetCurrentFrame() const
+FFrameRate FLiveLinkHubPlaybackController::GetFrameRate() const
 {
-	return CurrentFrameIndex;
+	return CurrentFrameRate;
 }
 
 void FLiveLinkHubPlaybackController::Start()
@@ -351,8 +387,9 @@ uint32 FLiveLinkHubPlaybackController::Run()
 				auto SetPlayhead = [&]()
 				{
 					const double Delta = FPlatformTime::Seconds() - PlaybackStartTime;
-					Playhead = bIsReverse ? StartTimestamp - Delta : StartTimestamp + Delta;
-					Playhead = FMath::Clamp(Playhead.load(), GetSelectionStartTime(), GetSelectionEndTime());
+					double Position = bIsReverse ? StartTimestamp - Delta : StartTimestamp + Delta;
+					Position = FMath::Clamp(Position, GetSelectionStartTime().AsSeconds(), GetSelectionEndTime().AsSeconds());
+					Playhead->SetValue(FQualifiedFrameTime(FFrameTime::FromDecimal(Position * GetFrameRate().Numerator), GetFrameRate()));
 				};
 
 				SetPlayhead();
@@ -406,11 +443,10 @@ void FLiveLinkHubPlaybackController::PushSubjectData(const FLiveLinkRecordedFram
 	}
 	else
 	{
-		// Record the frame index when pushing so it is accurate. If we only calculate based on time it may not match the actual frames sent.
-		CurrentFrameIndex = NextFrame.FrameIndex;
-		
 		FLiveLinkFrameDataStruct FrameDataStruct;
 		FrameDataStruct.InitializeWith(NextFrame.Data.GetScriptStruct(), (FLiveLinkBaseFrameData*)NextFrame.Data.GetMemory());
+
+		CurrentFrameRate = FrameDataStruct.GetBaseData()->MetaData.SceneTime.Rate;
 		
 		if (bForceSync)
 		{
@@ -422,7 +458,7 @@ void FLiveLinkHubPlaybackController::PushSubjectData(const FLiveLinkRecordedFram
 
 bool FLiveLinkHubPlaybackController::SyncToPlayhead()
 {
-	const double Timestamp = Playhead.load();
+	const double Timestamp = Playhead->GetValue().AsSeconds();
 	TArray<FLiveLinkRecordedFrame> NextFrames = bIsReverse ? RecordingPlayer->FetchPreviousFramesAtTimestamp(Timestamp)
 		: RecordingPlayer->FetchNextFramesAtTimestamp(Timestamp);
 	
@@ -436,9 +472,9 @@ bool FLiveLinkHubPlaybackController::SyncToPlayhead()
 	return NextFrames.Num() > 0;
 }
 
-bool FLiveLinkHubPlaybackController::SyncToFrame(int32 InFrameIndex)
+bool FLiveLinkHubPlaybackController::SyncToFrame(const FFrameNumber& InFrameNumber)
 {
-	TArray<FLiveLinkRecordedFrame> NextFrames = RecordingPlayer->FetchNextFramesAtIndex(InFrameIndex);
+	TArray<FLiveLinkRecordedFrame> NextFrames = RecordingPlayer->FetchNextFramesAtIndex(InFrameNumber.Value);
 	
 	for (const FLiveLinkRecordedFrame& NextFrame : NextFrames)
 	{
@@ -450,11 +486,7 @@ bool FLiveLinkHubPlaybackController::SyncToFrame(int32 InFrameIndex)
 
 bool FLiveLinkHubPlaybackController::ShouldRestart() const
 {
-	return RecordingToPlay.IsValid() && ((bIsReverse && CurrentFrameIndex <= GetSelectionStartFrame()) || (!bIsReverse && CurrentFrameIndex >= GetSelectionEndFrame()));
-}
-
-double FLiveLinkHubPlaybackController::GetTimeDelta() const
-{
-	// todo: get frame frate
-	return RecordingToPlay.IsValid() ? RecordingToPlay->LengthInSeconds / 60.f : 1.f;
+	const FFrameNumber CurrentFrame = GetCurrentFrame();
+	return RecordingToPlay.IsValid() && ((bIsReverse && CurrentFrame <= GetSelectionStartTime().Time.GetFrame())
+		|| (!bIsReverse && CurrentFrame >= GetSelectionEndTime().Time.GetFrame()));
 }
diff --git a/Engine/Plugins/Animation/LiveLinkHub/Source/LiveLinkHub/Private/Recording/LiveLinkHubPlaybackController.h b/Engine/Plugins/Animation/LiveLinkHub/Source/LiveLinkHub/Private/Recording/LiveLinkHubPlaybackController.h
--- a/Engine/Plugins/Animation/LiveLinkHub/Source/LiveLinkHub/Private/Recording/LiveLinkHubPlaybackController.h
+++ b/Engine/Plugins/Animation/LiveLinkHub/Source/LiveLinkHub/Private/Recording/LiveLinkHubPlaybackController.h
@@ -5,13 +5,11 @@
 #include "HAL/Runnable.h"
 
 #include "HAL/Event.h"
-#include "HAL/ThreadSafeBool.h"
-#include "Misc/Timespan.h"
-#include "Misc/Optional.h"
+#include "Misc/FrameNumber.h"
+#include "Misc/FrameRate.h"
+#include "Misc/QualifiedFrameTime.h"
 #include "Recording/LiveLinkRecordingPlayer.h"
-#include "Templates/SubclassOf.h"
 #include "Templates/SharedPointer.h"
-#include "Templates/PimplPtr.h"
 #include "UObject/StrongObjectPtr.h"
 
 #include <atomic>
@@ -23,6 +21,7 @@ class ULiveLinkRecording;
 class ULiveLinkPreset;
 class ULiveLinkRecording;
 class SWidget;
+class FLiveLinkHubAtomicQualifiedFrameTime;
 
 class FLiveLinkHubPlaybackController : public FRunnable
 {
@@ -58,44 +57,32 @@ public:
 	void Eject();
 
 	/** Go to a specific time. */
-	void GoToTime(double InTime);
-
-	/** Go to a specific frame index. */
-	void GoToFrame(int32 InFrameIndex);
-
-	/** Calculate the frame index for the given time. */
-	int32 GetFrameIndexFromTime(double InTime, bool bReverse = false) const;
-
-	/** Retrieve the timestamp from the given frame index. */
-	double GetTimeFromFrameIndex(int32 InFrameIndex) const;
-
-	/** Retrieve the selection start frame. */
-	int32 GetSelectionStartFrame() const;
-
-	/** Retrieve the selection end frame. */
-	int32 GetSelectionEndFrame() const;
-
+	void GoToTime(FQualifiedFrameTime InTime);
+	
 	/** Retrieve the selection start time. */
-	double GetSelectionStartTime() const;
+	FQualifiedFrameTime GetSelectionStartTime() const;
 
 	/** Set the selection start time. */
-	void SetSelectionStartTime(double InTime);
+	void SetSelectionStartTime(FQualifiedFrameTime InTime);
 
 	/** Retrieve the selection end time. */
-	double GetSelectionEndTime() const;
+	FQualifiedFrameTime GetSelectionEndTime() const;
 
 	/** Set the selection end time. */
-	void SetSelectionEndTime(double InTime);
+	void SetSelectionEndTime(FQualifiedFrameTime InTime);
 
 	/** Retrieve the length of the recording. */
-	double GetLength() const;
-
+	FQualifiedFrameTime GetLength() const;
+	
 	/** Retrieve the playhead. */
-	double GetCurrentTime() const;
+	FQualifiedFrameTime GetCurrentTime() const;
 
-	/** Return the exact frame index of frames that have been processed. */
-	int32 GetCurrentFrame() const;
+	/** Retrieve the current frame of the animation. */
+	FFrameNumber GetCurrentFrame() const;
 
+	/** Retrieve the current framerate. */
+	FFrameRate GetFrameRate() const;
+	
 	/** If the controller is ready for commands. */
 	bool IsReady() const
 	{
@@ -172,14 +159,11 @@ private:
 	bool SyncToPlayhead();
 
 	/** Force sync to a specific frame. */
-	bool SyncToFrame(int32 InFrameIndex);
+	bool SyncToFrame(const FFrameNumber& InFrameNumber);
 
 	/** Checks if the current playback settings indicates the recording should restart. */
 	bool ShouldRestart() const;
 
-	/** Retrieve the time delta to use for the recording when dragging the playhead spinbox. */
-	double GetTimeDelta() const;
-	
 private:
 	/** If the system has established a connection with the client. */
 	bool bIsReady = false;
@@ -215,17 +199,17 @@ private:
 	/** Time that the playback started. */
 	double PlaybackStartTime = 0.0;
 	/** Playhead for the current playback. */
-	std::atomic<double> Playhead = 0.0;
+	TSharedPtr<FLiveLinkHubAtomicQualifiedFrameTime, ESPMode::ThreadSafe> Playhead;
 
 	/** The view range of the slider, defaults to start/end time. */
 	TRange<double> SliderViewRange = TRange<double>(0.f, 0.f);
 
 	/** The playback selection start time. */
-	double SelectionStartTime = 0.f;
+	FQualifiedFrameTime SelectionStartTime;
 
 	/** The playback selection end time. */
-	double SelectionEndTime = 0.f;
-	
-	/** The index of the latest processed frame. This may not exactly match the index calculated from the playhead. */
-	int32 CurrentFrameIndex = INDEX_NONE;
+	FQualifiedFrameTime SelectionEndTime;
+
+	/** Current framerate of the recording, sampled from the latest frame. */
+	FFrameRate CurrentFrameRate;
 };
diff --git a/Engine/Plugins/Animation/LiveLinkHub/Source/LiveLinkHub/Private/Recording/LiveLinkRecordingPlayer.h b/Engine/Plugins/Animation/LiveLinkHub/Source/LiveLinkHub/Private/Recording/LiveLinkRecordingPlayer.h
--- a/Engine/Plugins/Animation/LiveLinkHub/Source/LiveLinkHub/Private/Recording/LiveLinkRecordingPlayer.h
+++ b/Engine/Plugins/Animation/LiveLinkHub/Source/LiveLinkHub/Private/Recording/LiveLinkRecordingPlayer.h
@@ -48,4 +48,7 @@ public:
 
 	/** Restart the recording from the beginning. */
 	virtual void RestartPlayback(int32 InIndex = INDEX_NONE) = 0;
+
+	/** Retrieve the first frame's frame rate information. */
+	virtual FFrameRate GetInitialFramerate() = 0;
 };
diff --git a/Engine/Plugins/Animation/LiveLinkHub/Source/LiveLinkHub/Private/UI/Widgets/SLiveLinkHubPlaybackWidget.cpp b/Engine/Plugins/Animation/LiveLinkHub/Source/LiveLinkHub/Private/UI/Widgets/SLiveLinkHubPlaybackWidget.cpp
--- a/Engine/Plugins/Animation/LiveLinkHub/Source/LiveLinkHub/Private/UI/Widgets/SLiveLinkHubPlaybackWidget.cpp
+++ b/Engine/Plugins/Animation/LiveLinkHub/Source/LiveLinkHub/Private/UI/Widgets/SLiveLinkHubPlaybackWidget.cpp
@@ -3,8 +3,10 @@
 #include "SLiveLinkHubPlaybackWidget.h"
 
 #include "FrameNumberNumericInterface.h"
+#include "Framework/MultiBox/MultiBoxBuilder.h"
 #include "SSimpleTimeSlider.h"
 #include "Widgets/Input/SButton.h"
+#include "Widgets/Input/SComboButton.h"
 #include "Widgets/Input/SSpinBox.h"
 
 #define LOCTEXT_NAMESPACE "SLiveLinkHubPlaybackWidget"
@@ -31,8 +33,6 @@ void SLiveLinkHubPlaybackWidget::Construct(const FArguments& InArgs)
 	OnGetViewRangeDelegate = InArgs._GetViewRange;
 	OnSetViewRangeDelegate = InArgs._SetViewRange;
 
-	OnGetCurrentFrameDelegate = InArgs._GetCurrentFrame;
-
 	OnGetSelectionStartTimeDelegate = InArgs._GetSelectionStartTime;
 	OnGetSelectionEndTimeDelegate = InArgs._GetSelectionEndTime;
 	OnSetSelectionStartTimeDelegate = InArgs._SetSelectionStartTime;
@@ -41,15 +41,15 @@ void SLiveLinkHubPlaybackWidget::Construct(const FArguments& InArgs)
 	OnSetLoopingDelegate = InArgs._OnSetLooping;
 	OnGetLoopingDelegate = InArgs._IsLooping;
 
-	OnGetTimeDelta = InArgs._GetTimeDelta;
+	OnGetFrameRate = InArgs._GetFrameRate;
 
-	TAttribute<EFrameNumberDisplayFormats> GetDisplayFormatAttr = TAttribute<EFrameNumberDisplayFormats>(this, &SLiveLinkHubPlaybackWidget::GetDisplayFormat);
-	TAttribute<FFrameRate> GetTickResolutionAttr = TAttribute<FFrameRate>(this, &SLiveLinkHubPlaybackWidget::GetFocusedTickResolution);
-	TAttribute<FFrameRate> GetDisplayRateAttr    = TAttribute<FFrameRate>(this, &SLiveLinkHubPlaybackWidget::GetFocusedDisplayRate);
+	const TAttribute<EFrameNumberDisplayFormats> GetDisplayFormatAttr = MakeAttributeSP(this, &SLiveLinkHubPlaybackWidget::GetDisplayFormat);
+	const TAttribute<FFrameRate> GetDisplayRateAttr = MakeAttributeSP(this, &SLiveLinkHubPlaybackWidget::GetFrameRate);
 	
 	// Create our numeric type interface so we can pass it to the time slider below.
-	TSharedPtr<FFrameNumberInterface> NumberInterface = MakeShareable(new FFrameNumberInterface(GetDisplayFormatAttr, 0, GetTickResolutionAttr, GetDisplayRateAttr));
-	
+	NumberInterface = MakeShareable(new FFrameNumberInterface(GetDisplayFormatAttr, 0, GetDisplayRateAttr, GetDisplayRateAttr));
+
+	NumberInterface->DisplayFormatChanged();
 	ChildSlot
 	[
 		SNew(SVerticalBox)
@@ -60,7 +60,12 @@ void SLiveLinkHubPlaybackWidget::Construct(const FArguments& InArgs)
 			SNew(SSimpleTimeSlider)
 			.ClampRangeHighlightSize(0.15f)
 			.ClampRangeHighlightColor(FLinearColor::Gray.CopyWithNewOpacity(0.5f))
-			.ScrubPosition(this, &SLiveLinkHubPlaybackWidget::GetCurrentTime)
+			.ScrubPosition_Lambda([this]()
+			{
+				// TimeSlider needs actual time in seconds, not the frame time.
+				const double Seconds = GetCurrentTime() / GetFrameRate().Numerator;
+				return Seconds;
+			})
 			.ViewRange(this, &SLiveLinkHubPlaybackWidget::GetViewRange)
 			.OnViewRangeChanged(this, &SLiveLinkHubPlaybackWidget::SetViewRange)
 			.ClampRange(this, &SLiveLinkHubPlaybackWidget::GetClampRange)
@@ -69,7 +74,8 @@ void SLiveLinkHubPlaybackWidget::Construct(const FArguments& InArgs)
 						{
 							if (bIsScrubbing)
 							{
-								SetCurrentTime(NewScrubTime);
+								//  Convert time in seconds to frame time as double.
+								SetCurrentTime(NewScrubTime * GetFrameRate().Numerator);
 							}
 						})
 		]
@@ -244,6 +250,17 @@ void SLiveLinkHubPlaybackWidget::Construct(const FArguments& InArgs)
 				.Style(&FAppStyle::Get().GetWidgetStyle<FSpinBoxStyle>("Sequencer.HyperlinkSpinBox"))
 				.TypeInterface(NumberInterface)
 			]
+			+SHorizontalBox::Slot()
+			.Padding(4.f, 0.f)
+			.AutoWidth()
+			[
+				SNew(SComboButton)
+				.HasDownArrow(true)
+				.OnGetMenuContent(this, &SLiveLinkHubPlaybackWidget::MakePlaybackSettingsDropdown)
+				.ForegroundColor(FSlateColor::UseStyle())
+				.ToolTipText(LOCTEXT("PlaybackSettings_Tooltip", "Change playback settings."))
+				.ComboButtonStyle(&FAppStyle::Get().GetWidgetStyle<FComboButtonStyle>("SimpleComboButton"))
+			]
 		]
 	];
 }
@@ -294,7 +311,10 @@ FReply SLiveLinkHubPlaybackWidget::OnLoopPressed()
 void SLiveLinkHubPlaybackWidget::SetCurrentTime(double InTime)
 {
 	check(OnSetCurrentTimeDelegate.IsBound());
-	OnSetCurrentTimeDelegate.Execute(InTime);
+
+	const FQualifiedFrameTime FrameTime = SecondsToFrameTime(InTime);
+
+	OnSetCurrentTimeDelegate.Execute(FrameTime);
 }
 
 void SLiveLinkHubPlaybackWidget::OnCurrentTimeCommitted(double InTime, ETextCommit::Type InTextCommit)
@@ -305,25 +325,23 @@ void SLiveLinkHubPlaybackWidget::OnCurrentTimeCommitted(double InTime, ETextComm
 double SLiveLinkHubPlaybackWidget::GetCurrentTime() const
 {
 	check(OnGetCurrentTimeDelegate.IsBound());
-	return OnGetCurrentTimeDelegate.Execute();
+	const FQualifiedFrameTime FrameTime = OnGetCurrentTimeDelegate.Execute();
+		
+	return FrameTime.Time.AsDecimal();
 }
 
 double SLiveLinkHubPlaybackWidget::GetTotalLength() const
 {
 	check(OnGetTotalLengthDelegate.IsBound());
-	return OnGetTotalLengthDelegate.Execute();
-}
-
-int32 SLiveLinkHubPlaybackWidget::GetCurrentFrame() const
-{
-	check(OnGetCurrentFrameDelegate.IsBound());
-	return OnGetCurrentFrameDelegate.Execute();
+	const FQualifiedFrameTime FrameTime = OnGetTotalLengthDelegate.Execute();
+	return FrameTime.Time.AsDecimal();
 }
 
 double SLiveLinkHubPlaybackWidget::GetSelectionStartTime() const
 {
 	check(OnGetSelectionStartTimeDelegate.IsBound());
-	return OnGetSelectionStartTimeDelegate.Execute();
+	const FQualifiedFrameTime FrameTime = OnGetSelectionStartTimeDelegate.Execute();
+	return FrameTime.Time.AsDecimal();
 }
 
 void SLiveLinkHubPlaybackWidget::OnSelectionStartTimeCommitted(double InTime, ETextCommit::Type InTextCommit)
@@ -334,13 +352,16 @@ void SLiveLinkHubPlaybackWidget::OnSelectionStartTimeCommitted(double InTime, ET
 void SLiveLinkHubPlaybackWidget::SetSelectionStartTime(double InTime)
 {
 	check(OnSetSelectionStartTimeDelegate.IsBound());
-	OnSetSelectionStartTimeDelegate.Execute(InTime);
+	const FQualifiedFrameTime FrameTime = SecondsToFrameTime(InTime);
+	OnSetSelectionStartTimeDelegate.Execute(FrameTime);
 }
 
 double SLiveLinkHubPlaybackWidget::GetSelectionEndTime() const
 {
 	check(OnGetSelectionEndTimeDelegate.IsBound());
-	return OnGetSelectionEndTimeDelegate.Execute();
+	const FQualifiedFrameTime FrameTime = OnGetSelectionEndTimeDelegate.Execute();
+
+	return FrameTime.Time.AsDecimal();
 }
 
 void SLiveLinkHubPlaybackWidget::OnSelectionEndTimeCommitted(double InTime, ETextCommit::Type InTextCommit)
@@ -351,7 +372,8 @@ void SLiveLinkHubPlaybackWidget::OnSelectionEndTimeCommitted(double InTime, ETex
 void SLiveLinkHubPlaybackWidget::SetSelectionEndTime(double InTime)
 {
 	check(OnSetSelectionEndTimeDelegate.IsBound());
-	OnSetSelectionEndTimeDelegate.Execute(InTime);
+	const FQualifiedFrameTime FrameTime = SecondsToFrameTime(InTime);
+	OnSetSelectionEndTimeDelegate.Execute(FrameTime);
 }
 
 TRange<double> SLiveLinkHubPlaybackWidget::GetViewRange() const
@@ -368,7 +390,10 @@ void SLiveLinkHubPlaybackWidget::SetViewRange(TRange<double> InRange)
 
 TRange<double> SLiveLinkHubPlaybackWidget::GetClampRange() const
 {
-	return TRange<double>(GetSelectionStartTime(), GetSelectionEndTime());
+	const FFrameRate FrameRate = GetFrameRate();
+	const double Start = GetSelectionStartTime() / FrameRate.Numerator;
+	const double End = GetSelectionEndTime() / FrameRate.Numerator;
+	return TRange<double>(Start, End);
 }
 
 bool SLiveLinkHubPlaybackWidget::IsPaused() const
@@ -383,8 +408,7 @@ bool SLiveLinkHubPlaybackWidget::IsPlayingInReverse() const
 
 double SLiveLinkHubPlaybackWidget::GetSpinboxDelta() const
 {
-	check(OnGetTimeDelta.IsBound());
-	return OnGetTimeDelta.Execute();
+	return GetFrameRate().AsDecimal() * GetFrameRate().AsInterval();
 }
 
 EFrameNumberDisplayFormats SLiveLinkHubPlaybackWidget::GetDisplayFormat() const
@@ -392,18 +416,33 @@ EFrameNumberDisplayFormats SLiveLinkHubPlaybackWidget::GetDisplayFormat() const
 	return DisplayFormat;
 }
 
-FFrameRate SLiveLinkHubPlaybackWidget::GetFocusedTickResolution() const
+void SLiveLinkHubPlaybackWidget::SetDisplayFormat(EFrameNumberDisplayFormats InDisplayFormat)
 {
-	// todo
-	const int32 TickResolutionValue = 1000 * 24;
-	return FFrameRate(1, 1);
+	DisplayFormat = InDisplayFormat;
+	NumberInterface->DisplayFormatChanged();
 }
 
-FFrameRate SLiveLinkHubPlaybackWidget::GetFocusedDisplayRate() const
+bool SLiveLinkHubPlaybackWidget::CompareDisplayFormat(EFrameNumberDisplayFormats InDisplayFormat) const
 {
-	// todo
-	const int32 SequenceFrameRate = 24;
-	return FFrameRate(SequenceFrameRate, 1);
+	return DisplayFormat == InDisplayFormat;
+}
+
+FText SLiveLinkHubPlaybackWidget::GetDisplayFormatAsText() const
+{
+	return DisplayFormat == EFrameNumberDisplayFormats::Frames ? LOCTEXT("DisplayFormat_TimeFrames", "Frames") :
+		DisplayFormat == EFrameNumberDisplayFormats::Seconds ? LOCTEXT("DisplayFormat_TimeSeconds", "Seconds") :
+		LOCTEXT("DisplayFormat_Timecode", "Timecode") ;
+}
+
+FFrameRate SLiveLinkHubPlaybackWidget::GetFrameRate() const
+{
+	check(OnGetFrameRate.IsBound());
+	return OnGetFrameRate.Execute();
+}
+
+FQualifiedFrameTime SLiveLinkHubPlaybackWidget::SecondsToFrameTime(double InTime) const
+{
+	return FQualifiedFrameTime(FFrameTime::FromDecimal(InTime), GetFrameRate());
 }
 
 const FSlateBrush* SLiveLinkHubPlaybackWidget::GetPlayForwardIcon() const
@@ -432,4 +471,55 @@ FText SLiveLinkHubPlaybackWidget::GetLoopTooltip() const
 	return bLooping ? LOCTEXT("Loop", "Loop") : LOCTEXT("NoLoop", "No looping");
 }
 
+TSharedRef<SWidget> SLiveLinkHubPlaybackWidget::MakePlaybackSettingsDropdown()
+{
+	constexpr bool bCloseAfterSelection = true;
+	FMenuBuilder MenuBuilder(bCloseAfterSelection, nullptr);
+
+	const FCanExecuteAction AlwaysExecute = FCanExecuteAction::CreateLambda([]{ return true; });
+
+	MenuBuilder.BeginSection("ShowTimeAsSection", LOCTEXT("ShowTimeAs", "Show Time As"));
+	
+	MenuBuilder.AddMenuEntry(
+		LOCTEXT("Menu_TimecodeLabel", "Timecode"),
+		LOCTEXT("Menu_TimecodeTooltip", "Display values in timecode format."),
+		FSlateIcon(),
+		FUIAction(
+			FExecuteAction::CreateSP(this, &SLiveLinkHubPlaybackWidget::SetDisplayFormat, EFrameNumberDisplayFormats::NonDropFrameTimecode),
+			AlwaysExecute,
+			FIsActionChecked::CreateSP(this, &SLiveLinkHubPlaybackWidget::CompareDisplayFormat, EFrameNumberDisplayFormats::NonDropFrameTimecode)
+		),
+		NAME_None,
+		EUserInterfaceActionType::RadioButton
+	);
+	MenuBuilder.AddMenuEntry(
+		LOCTEXT("Menu_FramesLabel", "Frames"),
+		LOCTEXT("Menu_FramesTooltip", "Display values as frame numbers."),
+		FSlateIcon(),
+		FUIAction(
+			FExecuteAction::CreateSP(this, &SLiveLinkHubPlaybackWidget::SetDisplayFormat, EFrameNumberDisplayFormats::Frames),
+			AlwaysExecute,
+			FIsActionChecked::CreateSP(this, &SLiveLinkHubPlaybackWidget::CompareDisplayFormat, EFrameNumberDisplayFormats::Frames)
+		),
+		NAME_None,
+		EUserInterfaceActionType::RadioButton
+	);
+	MenuBuilder.AddMenuEntry(
+		LOCTEXT("Menu_SecondsLabel", "Seconds"),
+		LOCTEXT("Menu_SecondsTooltip", "Display values in seconds."),
+		FSlateIcon(),
+		FUIAction(
+			FExecuteAction::CreateSP(this, &SLiveLinkHubPlaybackWidget::SetDisplayFormat, EFrameNumberDisplayFormats::Seconds),
+			AlwaysExecute,
+			FIsActionChecked::CreateSP(this, &SLiveLinkHubPlaybackWidget::CompareDisplayFormat, EFrameNumberDisplayFormats::Seconds)
+		),
+		NAME_None,
+		EUserInterfaceActionType::RadioButton
+	);
+	
+	MenuBuilder.EndSection();
+	
+	return MenuBuilder.MakeWidget();
+}
+
 #undef LOCTEXT_NAMESPACE
diff --git a/Engine/Plugins/Animation/LiveLinkHub/Source/LiveLinkHub/Private/UI/Widgets/SLiveLinkHubPlaybackWidget.h b/Engine/Plugins/Animation/LiveLinkHub/Source/LiveLinkHub/Private/UI/Widgets/SLiveLinkHubPlaybackWidget.h
--- a/Engine/Plugins/Animation/LiveLinkHub/Source/LiveLinkHub/Private/UI/Widgets/SLiveLinkHubPlaybackWidget.h
+++ b/Engine/Plugins/Animation/LiveLinkHub/Source/LiveLinkHub/Private/UI/Widgets/SLiveLinkHubPlaybackWidget.h
@@ -4,10 +4,12 @@
 
 #include "FrameNumberDisplayFormat.h"
 #include "Misc/FrameRate.h"
+#include "Misc/QualifiedFrameTime.h"
 #include "Widgets/SCompoundWidget.h"
 
 class FLiveLinkHubPlaybackController;
 class ULiveLinkRecording;
+struct FFrameNumberInterface;
 
 /**
  * The playback widget for controlling live link hub animations.
@@ -17,9 +19,9 @@ class SLiveLinkHubPlaybackWidget : public SCompoundWidget
 public:
 	DECLARE_DELEGATE_RetVal(bool, FOnGetChecked);
 	DECLARE_DELEGATE_OneParam(FOnSetChecked, bool);
-	DECLARE_DELEGATE_RetVal(double, FOnGetTime);
-	DECLARE_DELEGATE_OneParam(FOnSetTime, double);
-	DECLARE_DELEGATE_RetVal(int32, FOnGetFrame);
+	DECLARE_DELEGATE_RetVal(FQualifiedFrameTime, FOnGetTime);
+	DECLARE_DELEGATE_OneParam(FOnSetTime, FQualifiedFrameTime);
+	DECLARE_DELEGATE_RetVal(FFrameRate, FOnGetFrame);
 	DECLARE_DELEGATE_OneParam(FOnSetViewRange, const TRange<double>&)
 	DECLARE_DELEGATE_RetVal(TRange<double>, FOnGetViewRange)
 
@@ -50,8 +52,6 @@ public:
 	SLATE_EVENT(FOnGetTime, GetCurrentTime)
 	/** Set the current time of the recording. */
 	SLATE_EVENT(FOnSetTime, SetCurrentTime)
-	/** Get the current frame of the recording. */
-	SLATE_EVENT(FOnGetFrame, GetCurrentFrame)
 	/** Get the total length of the recording. */
 	SLATE_EVENT(FOnGetTime, GetTotalLength)
 	/** Get the selection start of the recording. */
@@ -62,15 +62,14 @@ public:
 	SLATE_EVENT(FOnGetTime, GetSelectionEndTime)
 	/** Set the selection end of the recording. */
 	SLATE_EVENT(FOnSetTime, SetSelectionEndTime)
-
+	/** Retrieve the frame rate for the recording. */
+	SLATE_EVENT(FOnGetFrame, GetFrameRate)
+	
 	/** Get the view range (visible selection range). */
 	SLATE_EVENT(FOnSetViewRange, SetViewRange)
 	/** Set the view range (visible selection range). */
 	SLATE_EVENT(FOnGetViewRange, GetViewRange)
 	
-	/** Resolution for spinbox increments. */
-	SLATE_EVENT(FOnGetTime, GetTimeDelta)
-	
 	SLATE_END_ARGS()
 
 	/**
@@ -103,8 +102,6 @@ private:
 	double GetCurrentTime() const;
 	/** The total length of the recording. */
 	double GetTotalLength() const;
-	/** The current frame being played. */
-	int32 GetCurrentFrame() const;
 
 	/** Retrieve the selection start time. */
 	double GetSelectionStartTime() const;
@@ -139,10 +136,20 @@ private:
 	/** Frame or seconds to display for text and input. */
 	EFrameNumberDisplayFormats GetDisplayFormat() const;
 
-	/** Tick resolution for number inputs. */
-	FFrameRate GetFocusedTickResolution() const;
-	/** The display rate for number inputs. */
-	FFrameRate GetFocusedDisplayRate() const;
+	/** Sets the display format. */
+	void SetDisplayFormat(EFrameNumberDisplayFormats InDisplayFormat);
+
+	/** Validates the current display format is set to the given format. */
+	bool CompareDisplayFormat(EFrameNumberDisplayFormats InDisplayFormat) const;
+
+	/** Retrieve the text version of the display format. */
+	FText GetDisplayFormatAsText() const;
+	
+	/** The frame rate, used for number inputs. */
+	FFrameRate GetFrameRate() const;
+
+	/** Convert raw seconds to frame time. */
+	FQualifiedFrameTime SecondsToFrameTime(double InTime) const;
 
 	/** The forward icon to use for the play forward button. */
 	const FSlateBrush* GetPlayForwardIcon() const;
@@ -152,6 +159,9 @@ private:
 	const FSlateBrush* GetLoopIcon() const;
 	/** The tooltip for the loop button. */
 	FText GetLoopTooltip() const;
+
+	/** Create the dropdown box for playback settings. */
+	TSharedRef<SWidget> MakePlaybackSettingsDropdown();
 private:
 	TWeakPtr<FLiveLinkHubPlaybackController> PlaybackController;
 	/** Delegate for pressing play forward. */
@@ -193,9 +203,6 @@ private:
 	/** Delegate for setting the current playhead time. */
 	FOnSetTime OnSetCurrentTimeDelegate;
 
-	/** Delegate for getting the current frame. */
-	FOnGetFrame OnGetCurrentFrameDelegate;
-
 	/** Delegate for getting the view range. */
 	FOnGetViewRange OnGetViewRangeDelegate;
 	
@@ -214,9 +221,12 @@ private:
 	/** Delegate for setting the selection start. */
 	FOnSetTime OnSetSelectionEndTimeDelegate;
 
-	/** Delegate for the spinbox delta. */
-	FOnGetTime OnGetTimeDelta;
+	/** Retrieve the frame rate. */
+	FOnGetFrame OnGetFrameRate;
+
+	/** The number interface for displaying the correct frame format. */
+	TSharedPtr<FFrameNumberInterface> NumberInterface;
 
 	/** The format to display/edit values in. */
-	EFrameNumberDisplayFormats DisplayFormat = EFrameNumberDisplayFormats::Seconds;
+	EFrameNumberDisplayFormats DisplayFormat = EFrameNumberDisplayFormats::NonDropFrameTimecode;
 };
\ No newline at end of file
