LiveLinkHub - Remove toolbar entry for LiveLink Panel
Add livelink subjects to the livelink status bar

diff --git a/Engine/Plugins/Animation/LiveLink/Source/LiveLinkEditor/Private/LiveLinkClientPanelViews.cpp b/Engine/Plugins/Animation/LiveLink/Source/LiveLinkEditor/Private/LiveLinkClientPanelViews.cpp
--- a/Engine/Plugins/Animation/LiveLink/Source/LiveLinkEditor/Private/LiveLinkClientPanelViews.cpp
+++ b/Engine/Plugins/Animation/LiveLink/Source/LiveLinkEditor/Private/LiveLinkClientPanelViews.cpp
@@ -3,8 +3,10 @@
 #include "LiveLinkClientPanelViews.h"
 
 #include "EditorFontGlyphs.h"
+#include "Features/IModularFeatures.h"
 #include "Framework/Commands/UICommandList.h"
 #include "IDetailsView.h"
+#include "ILiveLinkClient.h"
 #include "Internationalization/Text.h"
 #include "LiveLinkClient.h"
 #include "LiveLinkClientCommands.h"
@@ -297,6 +299,11 @@ private:
 		{
 			return EntryPtr->IsSubjectValid() ? GetDefault<ULiveLinkSettings>()->GetValidColor() : GetDefault<ULiveLinkSettings>()->GetInvalidColor();
 		}
+
+		if (!EntryPtr->IsSubjectEnabled() && EntryPtr->IsSubject())
+		{
+			UE_LOG(LogTemp, Warning, TEXT("Entry %s disabled"), *EntryPtr->SubjectKey.SubjectName.ToString());
+		}
 		return FLinearColor(0.f, 0.f, 0.f, 0.f);
 	}
 
@@ -566,6 +573,76 @@ bool FLiveLinkSubjectsView::CanRemoveSubject() const
 	return Selected.Num() > 0 && Selected[0] && Selected[0]->IsVirtualSubject();
 }
 
+void FLiveLinkSubjectsView::RefreshSubjects()
+{
+	TArray<FLiveLinkSubjectKey> SavedSelection;
+	{
+		TArray<FLiveLinkSubjectUIEntryPtr> SelectedItems = SubjectsTreeView->GetSelectedItems();
+		for (const FLiveLinkSubjectUIEntryPtr& SelectedItem : SelectedItems)
+		{
+			SavedSelection.Add(SelectedItem->SubjectKey);
+		}
+	}
+
+	if (IModularFeatures::Get().IsModularFeatureAvailable(ILiveLinkClient::ModularFeatureName))
+	{
+		if (ILiveLinkClient* Client = &IModularFeatures::Get().GetModularFeature<ILiveLinkClient>(ILiveLinkClient::ModularFeatureName))
+		{
+			TArray<FLiveLinkSubjectKey> SubjectKeys = Client->GetSubjects(true, true);
+			SubjectData.Reset();
+
+			TMap<FGuid, FLiveLinkSubjectUIEntryPtr> SourceHeaderItems;
+			TArray<FLiveLinkSubjectUIEntryPtr> AllItems;
+			AllItems.Reserve(SubjectKeys.Num());
+
+			for (const FLiveLinkSubjectKey& SubjectKey : SubjectKeys)
+			{
+				FLiveLinkSubjectUIEntryPtr Source;
+				if (FLiveLinkSubjectUIEntryPtr* SourcePtr = SourceHeaderItems.Find(SubjectKey.Source))
+				{
+					Source = *SourcePtr;
+				}
+				else
+				{
+					FLiveLinkSubjectKey SourceKey = SubjectKey;
+					SourceKey.SubjectName = NAME_None;
+					Source = MakeShared<FLiveLinkSubjectUIEntry>(SourceKey, static_cast<FLiveLinkClient*>(Client));
+					SubjectData.Add(Source);
+					SourceHeaderItems.Add(SubjectKey.Source) = Source;
+
+					SubjectsTreeView->SetItemExpansion(Source, true);
+					AllItems.Add(Source);
+				}
+
+				FLiveLinkSubjectUIEntryPtr SubjectEntry = MakeShared<FLiveLinkSubjectUIEntry>(SubjectKey, static_cast<FLiveLinkClient*>(Client));
+				Source->Children.Add(SubjectEntry);
+				AllItems.Add(SubjectEntry);
+			}
+
+			auto SortPredicate = [](const FLiveLinkSubjectUIEntryPtr& LHS, const FLiveLinkSubjectUIEntryPtr& RHS) {return LHS->GetItemText().CompareTo(RHS->GetItemText()) < 0; };
+			SubjectData.Sort(SortPredicate);
+			for (FLiveLinkSubjectUIEntryPtr& Subject : SubjectData)
+			{
+				Subject->Children.Sort(SortPredicate);
+			}
+
+			for (const FLiveLinkSubjectUIEntryPtr& Item : AllItems)
+			{
+				for (FLiveLinkSubjectKey& Selection : SavedSelection)
+				{
+					if (Item->SubjectKey == Selection)
+					{
+						SubjectsTreeView->SetItemSelection(Item, true);
+						break;
+					}
+				}
+			}
+
+			SubjectsTreeView->RequestTreeRefresh();
+		}
+	}
+}
+
 void FLiveLinkSubjectsView::CreateSubjectsTreeView(const TSharedPtr<FUICommandList>& InCommandList)
 {
 	SAssignNew(SubjectsTreeView, SLiveLinkSubjectsTreeView, bReadOnly)
diff --git a/Engine/Plugins/Animation/LiveLink/Source/LiveLinkEditor/Private/LiveLinkPanelController.cpp b/Engine/Plugins/Animation/LiveLink/Source/LiveLinkEditor/Private/LiveLinkPanelController.cpp
--- a/Engine/Plugins/Animation/LiveLink/Source/LiveLinkEditor/Private/LiveLinkPanelController.cpp
+++ b/Engine/Plugins/Animation/LiveLink/Source/LiveLinkEditor/Private/LiveLinkPanelController.cpp
@@ -185,66 +185,10 @@ void FLiveLinkPanelController::RebuildSourceList()
 
 void FLiveLinkPanelController::RebuildSubjectList()
 {
-	TArray<FLiveLinkSubjectKey> SavedSelection;
+	if (SubjectsView)
 	{
-		TArray<FLiveLinkSubjectUIEntryPtr> SelectedItems = SubjectsView->SubjectsTreeView->GetSelectedItems();
-		for (const FLiveLinkSubjectUIEntryPtr& SelectedItem : SelectedItems)
-		{
-			SavedSelection.Add(SelectedItem->SubjectKey);
-		}
-	}
-
-	TArray<FLiveLinkSubjectKey> SubjectKeys = Client->GetSubjects(true, true);
-	SubjectsView->SubjectData.Reset();
-
-	TMap<FGuid, FLiveLinkSubjectUIEntryPtr> SourceHeaderItems;
-	TArray<FLiveLinkSubjectUIEntryPtr> AllItems;
-	AllItems.Reserve(SubjectKeys.Num());
-
-	for (const FLiveLinkSubjectKey& SubjectKey : SubjectKeys)
-	{
-		FLiveLinkSubjectUIEntryPtr Source;
-		if(FLiveLinkSubjectUIEntryPtr* SourcePtr = SourceHeaderItems.Find(SubjectKey.Source))
-		{
-			Source = *SourcePtr;
-		}
-		else
-		{
-			FLiveLinkSubjectKey SourceKey = SubjectKey;
-			SourceKey.SubjectName = NAME_None;
-			Source = MakeShared<FLiveLinkSubjectUIEntry>(SourceKey, Client);
-			SubjectsView->SubjectData.Add(Source);
-			SourceHeaderItems.Add(SubjectKey.Source) = Source;
-
-			SubjectsView->SubjectsTreeView->SetItemExpansion(Source, true);
-			AllItems.Add(Source);
-		}
-
-		FLiveLinkSubjectUIEntryPtr SubjectEntry = MakeShared<FLiveLinkSubjectUIEntry>(SubjectKey, Client);
-		Source->Children.Add(SubjectEntry);
-		AllItems.Add(SubjectEntry);
-	}
-
-	auto SortPredicate = [](const FLiveLinkSubjectUIEntryPtr& LHS, const FLiveLinkSubjectUIEntryPtr& RHS) {return LHS->GetItemText().CompareTo(RHS->GetItemText()) < 0; };
-	SubjectsView->SubjectData.Sort(SortPredicate);
-	for (FLiveLinkSubjectUIEntryPtr& Subject : SubjectsView->SubjectData)
-	{
-		Subject->Children.Sort(SortPredicate);
+		SubjectsView->RefreshSubjects();
 	}
-
-	for (const FLiveLinkSubjectUIEntryPtr& Item : AllItems)
-	{
-		for (FLiveLinkSubjectKey& Selection : SavedSelection)
-		{
-			if (Item->SubjectKey == Selection)
-			{
-				SubjectsView->SubjectsTreeView->SetItemSelection(Item, true);
-				break;
-			}
-		}
-	}
-
-	SubjectsView->SubjectsTreeView->RequestTreeRefresh();
 }
 
 bool FLiveLinkPanelController::HasSource() const
diff --git a/Engine/Plugins/Animation/LiveLink/Source/LiveLinkEditor/Public/LiveLinkClientPanelViews.h b/Engine/Plugins/Animation/LiveLink/Source/LiveLinkEditor/Public/LiveLinkClientPanelViews.h
--- a/Engine/Plugins/Animation/LiveLink/Source/LiveLinkEditor/Public/LiveLinkClientPanelViews.h
+++ b/Engine/Plugins/Animation/LiveLink/Source/LiveLinkEditor/Public/LiveLinkClientPanelViews.h
@@ -158,7 +158,7 @@ public:
 	}
 };
 
-class FLiveLinkSubjectsView : public TSharedFromThis<FLiveLinkSubjectsView>
+class LIVELINKEDITOR_API FLiveLinkSubjectsView : public TSharedFromThis<FLiveLinkSubjectsView>
 {
 public:
 	DECLARE_DELEGATE_TwoParams(FOnSubjectSelectionChanged, FLiveLinkSubjectUIEntryPtr, ESelectInfo::Type);
@@ -175,6 +175,8 @@ public:
 	TSharedPtr<SWidget> OnOpenVirtualSubjectContextMenu(TSharedPtr<FUICommandList> InCommandList);
 	// Return whether a subject can be removed
 	bool CanRemoveSubject() const;
+	// Refresh the list of subjects using the livelink client.
+	void RefreshSubjects();
 
 private:
 	// Create the subjects tree view
diff --git a/Engine/Plugins/Animation/LiveLinkHub/Source/LiveLinkHubEditor/LiveLinkHubEditor.Build.cs b/Engine/Plugins/Animation/LiveLinkHub/Source/LiveLinkHubEditor/LiveLinkHubEditor.Build.cs
--- a/Engine/Plugins/Animation/LiveLinkHub/Source/LiveLinkHubEditor/LiveLinkHubEditor.Build.cs
+++ b/Engine/Plugins/Animation/LiveLinkHub/Source/LiveLinkHubEditor/LiveLinkHubEditor.Build.cs
@@ -18,6 +18,7 @@ namespace UnrealBuildTool.Rules
 			{
 				"Engine",
 				"LiveLink",
+				"LiveLinkEditor",
 				"LiveLinkHub",
 				"LiveLinkHubMessaging",
 				"LiveLinkInterface",
diff --git a/Engine/Plugins/Animation/LiveLinkHub/Source/LiveLinkHubEditor/Private/LiveLinkHubEditorModule.cpp b/Engine/Plugins/Animation/LiveLinkHub/Source/LiveLinkHubEditor/Private/LiveLinkHubEditorModule.cpp
--- a/Engine/Plugins/Animation/LiveLinkHub/Source/LiveLinkHubEditor/Private/LiveLinkHubEditorModule.cpp
+++ b/Engine/Plugins/Animation/LiveLinkHub/Source/LiveLinkHubEditor/Private/LiveLinkHubEditorModule.cpp
@@ -5,6 +5,28 @@
 #include "SLiveLinkHubEditorStatusBar.h"
 #include "ToolMenus.h"
 
+static TAutoConsoleVariable<int32> CVarLiveLinkHubEnableStatusBar(
+	TEXT("LiveLinkHub.EnableStatusBar"), 1,
+	TEXT("Whether to enable showing the livelink hub status bar in the editor. Must be set before launching the editor."),
+	ECVF_RenderThreadSafe);
+
+void FLiveLinkHubEditorModule::StartupModule()
+{
+	if (!IsRunningCommandlet() && CVarLiveLinkHubEnableStatusBar.GetValueOnAnyThread())
+	{
+		FCoreDelegates::OnPostEngineInit.AddRaw(this, &FLiveLinkHubEditorModule::OnPostEngineInit);
+	}
+}
+
+void FLiveLinkHubEditorModule::ShutdownModule()
+{
+	if (!IsRunningCommandlet() && CVarLiveLinkHubEnableStatusBar.GetValueOnAnyThread())
+	{
+		FCoreDelegates::OnPostEngineInit.RemoveAll(this);
+		UnregisterLiveLinkHubStatusBar();
+	}
+}
+
 void FLiveLinkHubEditorModule::OnPostEngineInit()
 {
 	if (GEditor)
diff --git a/Engine/Plugins/Animation/LiveLinkHub/Source/LiveLinkHubEditor/Private/LiveLinkHubEditorModule.h b/Engine/Plugins/Animation/LiveLinkHub/Source/LiveLinkHubEditor/Private/LiveLinkHubEditorModule.h
--- a/Engine/Plugins/Animation/LiveLinkHub/Source/LiveLinkHubEditor/Private/LiveLinkHubEditorModule.h
+++ b/Engine/Plugins/Animation/LiveLinkHub/Source/LiveLinkHubEditor/Private/LiveLinkHubEditorModule.h
@@ -14,22 +14,9 @@ class FLiveLinkHubEditorModule : public IModuleInterface
 
 public:
 	//~ Begin IModuleInterface interface
-	virtual void StartupModule()
-	{
-		if (!IsRunningCommandlet())
-		{
-			FCoreDelegates::OnPostEngineInit.AddRaw(this, &FLiveLinkHubEditorModule::OnPostEngineInit);
-		}
-	}
-
-	virtual void ShutdownModule()
-	{
-		if (!IsRunningCommandlet())
-		{
-			FCoreDelegates::OnPostEngineInit.RemoveAll(this);
-			UnregisterLiveLinkHubStatusBar();
-		}
-	}
+	virtual void StartupModule();
+
+	virtual void ShutdownModule();
 	//~ End IModuleInterface interface
 
 private:
diff --git a/Engine/Plugins/Animation/LiveLinkHub/Source/LiveLinkHubEditor/Private/SLiveLinkHubEditorStatusBar.cpp b/Engine/Plugins/Animation/LiveLinkHub/Source/LiveLinkHubEditor/Private/SLiveLinkHubEditorStatusBar.cpp
--- a/Engine/Plugins/Animation/LiveLinkHub/Source/LiveLinkHubEditor/Private/SLiveLinkHubEditorStatusBar.cpp
+++ b/Engine/Plugins/Animation/LiveLinkHub/Source/LiveLinkHubEditor/Private/SLiveLinkHubEditorStatusBar.cpp
@@ -8,6 +8,8 @@
 #include "ILiveLinkHubMessagingModule.h"
 #include "ILiveLinkModule.h"
 #include "ILiveLinkSource.h"
+#include "LiveLinkClientPanelViews.h"
+#include "LiveLinkEditorSettings.h"
 #include "LiveLinkMessageBusSource.h"
 #include "LiveLinkPreset.h"
 #include "Styling/SlateIconFinder.h"
@@ -15,6 +17,7 @@
 #include "TimerManager.h"
 #include "Widgets/Images/SImage.h"
 #include "Widgets/Input/SComboButton.h"
+#include "Widgets/Layout/SBox.h"
 #include "Widgets/Text/STextBlock.h"
 
 #define LOCTEXT_NAMESPACE "LiveLinkHubStatusBar"
@@ -26,6 +29,9 @@ void SLiveLinkHubEditorStatusBar::Construct(const FArguments& InArgs)
 
 	LiveLinkClient = &IModularFeatures::Get().GetModularFeature<ILiveLinkClient>(ILiveLinkClient::ModularFeatureName);
 	LiveLinkClient->OnLiveLinkSourceRemoved().AddSP(this, &SLiveLinkHubEditorStatusBar::OnSourceRemoved);
+	LiveLinkClient->OnLiveLinkSubjectRemoved().AddSP(this, &SLiveLinkHubEditorStatusBar::OnSubjectRemoved);
+	LiveLinkClient->OnLiveLinkSourcesChanged().AddRaw(this, &SLiveLinkHubEditorStatusBar::RefreshSubjects);
+	LiveLinkClient->OnLiveLinkSubjectsChanged().AddRaw(this, &SLiveLinkHubEditorStatusBar::RefreshSubjects);
 
 	constexpr bool bLoop = true;
 	if (GEditor && GEditor->IsTimerManagerValid())
@@ -33,6 +39,22 @@ void SLiveLinkHubEditorStatusBar::Construct(const FArguments& InArgs)
 		GEditor->GetTimerManager()->SetTimer(TimerHandle, FTimerDelegate::CreateSP(this, &SLiveLinkHubEditorStatusBar::CheckHubConnection), CheckConnectionIntervalSeconds, bLoop);
 	}
 
+	auto NullLambda = [](TSharedPtr<class FLiveLinkSubjectUIEntry>, ESelectInfo::Type) {};
+	TWeakPtr<SLiveLinkHubEditorStatusBar> WeakStatusBar = StaticCastSharedRef<SLiveLinkHubEditorStatusBar>(AsShared());
+	TAttribute<bool> ReadOnlyLambda = TAttribute<bool>::CreateLambda([WeakStatusBar]()
+	{
+		if (TSharedPtr<SLiveLinkHubEditorStatusBar> StatusBar = WeakStatusBar.Pin())
+		{
+			return StatusBar->ConnectionState == EHubConnectionState::Timeout || GetDefault<ULiveLinkEditorSettings>()->bReadOnly;
+		}
+		return true;
+	});
+
+
+	SubjectsView = MakeShared<FLiveLinkSubjectsView>(FLiveLinkSubjectsView::FOnSubjectSelectionChanged::CreateLambda(MoveTemp(NullLambda)), nullptr, ReadOnlyLambda);
+
+	RefreshSubjects();
+
 	ChildSlot
 	[
 		SNew(SComboButton)
@@ -63,11 +85,28 @@ void SLiveLinkHubEditorStatusBar::Construct(const FArguments& InArgs)
                 	.ToolTipText(this, &SLiveLinkHubEditorStatusBar::GetToolTipText)
                 ]
     		]
+			.MenuContent()
+			[
+				SNew(SBox)
+				.MinDesiredHeight(200.0f)
+				.MinDesiredWidth(300.0f)
+				[
+					SubjectsView->SubjectsTreeView.ToSharedRef()
+				]
+			]
 	];
 }
 
 SLiveLinkHubEditorStatusBar::~SLiveLinkHubEditorStatusBar()
 {
+	if (LiveLinkClient)
+	{
+		LiveLinkClient->OnLiveLinkSubjectsChanged().RemoveAll(this);
+		LiveLinkClient->OnLiveLinkSourcesChanged().RemoveAll(this);
+		LiveLinkClient->OnLiveLinkSubjectRemoved().RemoveAll(this);
+		LiveLinkClient->OnLiveLinkSourceRemoved().RemoveAll(this);
+	}
+
 	if (ILiveLinkHubMessagingModule* HubMessagingModule = FModuleManager::Get().GetModulePtr<ILiveLinkHubMessagingModule>("LiveLinkHubMessaging"))
 	{
 		HubMessagingModule->OnConnectionEstablished().RemoveAll(this);
@@ -105,6 +144,21 @@ void SLiveLinkHubEditorStatusBar::OnSourceRemoved(FGuid SourceId)
 	{
 		HubSourceId.Invalidate();
 	}
+
+	RefreshSubjects();
+}
+
+void SLiveLinkHubEditorStatusBar::OnSubjectRemoved(FLiveLinkSubjectKey)
+{
+	RefreshSubjects();
+}
+
+void SLiveLinkHubEditorStatusBar::RefreshSubjects()
+{
+	if (SubjectsView)
+	{
+		SubjectsView->RefreshSubjects();
+	}
 }
 
 void SLiveLinkHubEditorStatusBar::OnHubConnectionEstablished(FGuid SourceId)
diff --git a/Engine/Plugins/Animation/LiveLinkHub/Source/LiveLinkHubEditor/Private/SLiveLinkHubEditorStatusBar.h b/Engine/Plugins/Animation/LiveLinkHub/Source/LiveLinkHubEditor/Private/SLiveLinkHubEditorStatusBar.h
--- a/Engine/Plugins/Animation/LiveLinkHub/Source/LiveLinkHubEditor/Private/SLiveLinkHubEditorStatusBar.h
+++ b/Engine/Plugins/Animation/LiveLinkHub/Source/LiveLinkHubEditor/Private/SLiveLinkHubEditorStatusBar.h
@@ -4,6 +4,7 @@
 
 #include "Engine/TimerHandle.h"
 #include "Internationalization/Text.h"
+#include "LiveLinkTypes.h"
 #include "Templates/SharedPointer.h"
 #include "Types/SlateEnums.h"
 #include "Widgets/DeclarativeSyntaxSupport.h"
@@ -39,6 +40,9 @@ private:
 	/** Handler used to update the current connection state. */
 	void OnSourceRemoved(FGuid SourceId);
 
+	/** Handler used to update the subject list displayed by the status bar. */
+	void OnSubjectRemoved(FLiveLinkSubjectKey SubjectKey);
+
 	/** Handler used to update the connection state and source id when a connection with a hub instance is established. */
 	void OnHubConnectionEstablished(FGuid SourceId);
 
@@ -48,6 +52,9 @@ private:
 	/** Returns the visibility of the widget (Depends on whether we've established connection). */
 	EVisibility GetVisibility() const;
 
+	/** Method called when the subjects or sources have been modified, used to refresh the subject list UI. */
+	void RefreshSubjects();
+
 private:
 	/** LiveLink client used to retrieve the status of the hub connection. */
 	class ILiveLinkClient* LiveLinkClient = nullptr;
@@ -66,4 +73,7 @@ private:
 
 	/** Interval of the timer to check for connection validity. */
 	static constexpr float CheckConnectionIntervalSeconds = 1.0f;
+
+	/** Object responsible for creating and handling the subject list widget. */
+	TSharedPtr<class FLiveLinkSubjectsView> SubjectsView;
 };
