LiveLinkHub - Fix Linux compilation error

diff --git a/Engine/Plugins/Animation/LiveLinkHub/Source/LiveLinkHubEditor/Private/SLiveLinkHubEditorStatusBar.cpp b/Engine/Plugins/Animation/LiveLinkHub/Source/LiveLinkHubEditor/Private/SLiveLinkHubEditorStatusBar.cpp
--- a/Engine/Plugins/Animation/LiveLinkHub/Source/LiveLinkHubEditor/Private/SLiveLinkHubEditorStatusBar.cpp
+++ b/Engine/Plugins/Animation/LiveLinkHub/Source/LiveLinkHubEditor/Private/SLiveLinkHubEditorStatusBar.cpp
@@ -39,7 +39,7 @@ void SLiveLinkHubEditorStatusBar::Construct(const FArguments& InArgs)
 		GEditor->GetTimerManager()->SetTimer(TimerHandle, FTimerDelegate::CreateSP(this, &SLiveLinkHubEditorStatusBar::CheckHubConnection), CheckConnectionIntervalSeconds, bLoop);
 	}
 
-	auto NullLambda = [](TSharedPtr<class FLiveLinkSubjectUIEntry>, ESelectInfo::Type) {};
+	auto NullLambda = [](TSharedPtr<struct FLiveLinkSubjectUIEntry>, ESelectInfo::Type) {};
 	TWeakPtr<SLiveLinkHubEditorStatusBar> WeakStatusBar = StaticCastSharedRef<SLiveLinkHubEditorStatusBar>(AsShared());
 	TAttribute<bool> ReadOnlyLambda = TAttribute<bool>::CreateLambda([WeakStatusBar]()
 	{
