LiveLinkHub - Fix crash when changing map while receiving timecode through LiveLinkHub

- Use ExecuteOnGameThread when assigning a new timecode provider to make sure that the timecode assignment does not happen during GC.

diff --git a/Engine/Plugins/Animation/LiveLinkHub/Source/LiveLinkHubMessaging/Private/LiveLinkHubMessageBusSource.cpp b/Engine/Plugins/Animation/LiveLinkHub/Source/LiveLinkHubMessaging/Private/LiveLinkHubMessageBusSource.cpp
--- a/Engine/Plugins/Animation/LiveLinkHub/Source/LiveLinkHubMessaging/Private/LiveLinkHubMessageBusSource.cpp
+++ b/Engine/Plugins/Animation/LiveLinkHub/Source/LiveLinkHubMessaging/Private/LiveLinkHubMessageBusSource.cpp
@@ -2,6 +2,7 @@
 
 #include "LiveLinkHubMessageBusSource.h"
 
+#include "Containers/Ticker.h"
 #include "Engine/Level.h"
 #include "Engine/SystemTimeTimecodeProvider.h"
 #include "Engine/World.h"
@@ -79,7 +80,9 @@ void FLiveLinkHubMessageBusSource::InitializeMessageEndpoint(FMessageEndpointBui
 
 void FLiveLinkHubMessageBusSource::HandleTimecodeSettings(const FLiveLinkHubTimecodeSettings& Message, const TSharedRef<IMessageContext, ESPMode::ThreadSafe>& Context)
 {
-	Message.AssignTimecodeSettingsAsProviderToEngine();
+	ExecuteOnGameThread(UE_SOURCE_LOCATION, [Message]() {
+		Message.AssignTimecodeSettingsAsProviderToEngine();
+	});
 }
 
 FLiveLinkClientInfoMessage FLiveLinkHubMessageBusSource::CreateLiveLinkClientInfo() const
