LiveLinkHub - Fix recordings not being found when launched through a cooked editor

Added a mount point to the user content folder and added EnabledDataSources to the defaulteditor config to enable asset discovery.

diff --git a/Engine/Plugins/Animation/LiveLinkHub/Source/LiveLinkHub/LiveLinkHub.Build.cs b/Engine/Plugins/Animation/LiveLinkHub/Source/LiveLinkHub/LiveLinkHub.Build.cs
--- a/Engine/Plugins/Animation/LiveLinkHub/Source/LiveLinkHub/LiveLinkHub.Build.cs
+++ b/Engine/Plugins/Animation/LiveLinkHub/Source/LiveLinkHub/LiveLinkHub.Build.cs
@@ -17,6 +17,7 @@ namespace UnrealBuildTool.Rules
 			new string[]
 			{
 				"ApplicationCore",
+				"AssetRegistry",
 				"AssetTools",
 				"ContentBrowser",
 				"ContentBrowserAssetDataSource",
diff --git a/Engine/Plugins/Animation/LiveLinkHub/Source/LiveLinkHub/Private/LiveLinkHub.cpp b/Engine/Plugins/Animation/LiveLinkHub/Source/LiveLinkHub/Private/LiveLinkHub.cpp
--- a/Engine/Plugins/Animation/LiveLinkHub/Source/LiveLinkHub/Private/LiveLinkHub.cpp
+++ b/Engine/Plugins/Animation/LiveLinkHub/Source/LiveLinkHub/Private/LiveLinkHub.cpp
@@ -2,6 +2,8 @@
 
 #include "LiveLinkHub.h"
 
+#include "AssetRegistry/AssetRegistryModule.h"
+#include "AssetRegistry/IAssetRegistry.h"
 #include "Clients/LiveLinkHubClientsController.h"
 #include "Clients/LiveLinkHubProvider.h"
 #include "Config/LiveLinkHubFileUtilities.h"
@@ -16,6 +18,7 @@
 #include "LiveLinkProviderImpl.h"
 #include "LiveLinkSubject.h"
 #include "LiveLinkSubjectSettings.h"
+#include "Misc/App.h"
 #include "Misc/ConfigCacheIni.h"
 #include "Recording/LiveLinkHubPlaybackController.h"
 #include "Recording/LiveLinkHubRecordingController.h"
@@ -40,6 +43,10 @@ void FLiveLinkHub::Initialize()
 	SessionManager = MakeShared<FLiveLinkHubSessionManager>();
 	LiveLinkProvider = MakeShared<FLiveLinkHubProvider>(SessionManager.ToSharedRef());
 
+	IAssetRegistry& AssetRegistry = FModuleManager::LoadModuleChecked<FAssetRegistryModule>(AssetRegistryConstants::ModuleName).Get();
+	const FString FilePath = FPaths::Combine(FPlatformProcess::UserSettingsDir(), *FApp::GetEpicProductIdentifier(), TEXT("LiveLinkHub"), TEXT("Content"));
+	AssetRegistry.ScanPathsSynchronous({ FilePath }, /*bForceRescan=*/ true);
+
 	CommandExecutor = MakeUnique<FConsoleCommandExecutor>();
 	IModularFeatures::Get().RegisterModularFeature(IConsoleCommandExecutor::ModularFeatureName(), CommandExecutor.Get());
 
diff --git a/Engine/Programs/LiveLinkHub/Config/DefaultEditor.ini b/Engine/Programs/LiveLinkHub/Config/DefaultEditor.ini
new file mode 100644
--- /dev/null
+++ b/Engine/Programs/LiveLinkHub/Config/DefaultEditor.ini
@@ -0,0 +1,4 @@
+[/Script/ContentBrowserData.ContentBrowserDataSubsystem]
++EnabledDataSources="AssetData"
++EnabledDataSources="ClassData"
++EnabledDataSources="PortableObjectData"
diff --git a/Engine/Source/Programs/LiveLinkHub/Private/LiveLinkHubRun.cpp b/Engine/Source/Programs/LiveLinkHub/Private/LiveLinkHubRun.cpp
--- a/Engine/Source/Programs/LiveLinkHub/Private/LiveLinkHubRun.cpp
+++ b/Engine/Source/Programs/LiveLinkHub/Private/LiveLinkHubRun.cpp
@@ -2,6 +2,7 @@
 
 #include "LiveLinkHubRun.h"
 
+#include "Misc/App.h"
 #include "Misc/PackageName.h"
 #include "HAL/FileManager.h"
 #include "HAL/PlatformApplicationMisc.h"
@@ -64,6 +65,10 @@ int32 RunLiveLinkHub(const TCHAR* CommandLine)
 	// Map the engine's content to LiveLinkHub's content folder
 	FPackageName::RegisterMountPoint(TEXT("/Engine/"), RemappedFolder / TEXT("Content"));
 	FPackageName::RegisterMountPoint(TEXT("/Engine/Plugins/"), RemappedFolder / TEXT("Plugins"));
+
+	// This is needed to ensure that content generated by the hub will be discoverable (ie. LiveLink recordings)
+	const FString GameContentPath = FPaths::Combine(FPlatformProcess::UserSettingsDir(), *FApp::GetEpicProductIdentifier(), TEXT("LiveLinkHub"), TEXT("Content"));
+	FPackageName::RegisterMountPoint(TEXT("/Game/"), GameContentPath);
 	
 	// Needed to avoid assertion in scoped slow task
 	GWarn = FPlatformApplicationMisc::GetFeedbackContext();
