LiveLinkHub - Fix high DPI mode not working

Fix placeholder IP showing up in client details in the hub

diff --git a/Engine/Plugins/Animation/LiveLinkHub/Source/LiveLinkHub/Private/Clients/LiveLinkHubUEClientInfo.h b/Engine/Plugins/Animation/LiveLinkHub/Source/LiveLinkHub/Private/Clients/LiveLinkHubUEClientInfo.h
--- a/Engine/Plugins/Animation/LiveLinkHub/Source/LiveLinkHub/Private/Clients/LiveLinkHubUEClientInfo.h
+++ b/Engine/Plugins/Animation/LiveLinkHub/Source/LiveLinkHub/Private/Clients/LiveLinkHubUEClientInfo.h
@@ -88,7 +88,7 @@ struct FLiveLinkHubUEClientInfo
 	UPROPERTY(transient)
 	ELiveLinkClientStatus Status = ELiveLinkClientStatus::Disconnected;
 	
-	UPROPERTY(VisibleAnywhere, Category = "Client Details", DisplayName = "IP Address")
+	UPROPERTY()
 	FString IPAddress;
 	
 	/** Name of the host of the UE client */
diff --git a/Engine/Plugins/Animation/LiveLinkHub/Source/LiveLinkHub/Private/UI/Window/LiveLinkHubWindowController.cpp b/Engine/Plugins/Animation/LiveLinkHub/Source/LiveLinkHub/Private/UI/Window/LiveLinkHubWindowController.cpp
--- a/Engine/Plugins/Animation/LiveLinkHub/Source/LiveLinkHub/Private/UI/Window/LiveLinkHubWindowController.cpp
+++ b/Engine/Plugins/Animation/LiveLinkHub/Source/LiveLinkHub/Private/UI/Window/LiveLinkHubWindowController.cpp
@@ -2,6 +2,7 @@
 
 #include "LiveLinkHubWindowController.h"
 
+#include "CoreGlobals.h"
 #include "Framework/Application/SlateApplication.h"
 #include "Framework/Docking/LayoutService.h"
 #include "Framework/Notifications/NotificationManager.h"
@@ -98,7 +99,13 @@ TSharedPtr<FModalWindowManager> FLiveLinkHubWindowController::InitializeSlateApp
 {
 	TRACE_CPUPROFILER_EVENT_SCOPE(FLiveLinkHubWindowController::InitializeAsStandaloneApplication);
 	FSlateApplication::InitializeAsStandaloneApplication(GetStandardStandaloneRenderer());
+
+	// @hack This call will silently fail since we're running a commandlet, so pretend like we aren't one for it.
+	// In the future the slate code should be changed to allow visual commandlets to enable high dpi mode.
+	const bool bIsRunningCommandlet = PRIVATE_GIsRunningCommandlet;
+	PRIVATE_GIsRunningCommandlet = false;
 	FSlateApplication::InitHighDPI(true);
+	PRIVATE_GIsRunningCommandlet = bIsRunningCommandlet;
 
 	const FText ApplicationTitle = LOCTEXT("AppTitle", "LiveLink Hub");
 	FGlobalTabmanager::Get()->SetApplicationTitle(ApplicationTitle);
