LiveLinkHub - Fix data conflict when loading playback take alongside live data stream.

- Starting playback will properly close existing streams. When in playback new live streams can no longer be added. (UE-208520)
- Fix being able to record while in playback. (UE-209137)
- Fix deleting the recording source of the current playback not properly ejecting the playback. (UE-208525)

diff --git a/Engine/Plugins/Animation/LiveLinkHub/Source/LiveLinkHub/Private/Recording/LiveLinkHubPlaybackController.cpp b/Engine/Plugins/Animation/LiveLinkHub/Source/LiveLinkHub/Private/Recording/LiveLinkHubPlaybackController.cpp
--- a/Engine/Plugins/Animation/LiveLinkHub/Source/LiveLinkHub/Private/Recording/LiveLinkHubPlaybackController.cpp
+++ b/Engine/Plugins/Animation/LiveLinkHub/Source/LiveLinkHub/Private/Recording/LiveLinkHubPlaybackController.cpp
@@ -13,6 +13,7 @@
 #include "Features/IModularFeatures.h"
 #include "HAL/RunnableThread.h"
 #include "HAL/Event.h"
+#include "LiveLinkClient.h"
 #include "UI/Widgets/SLiveLinkHubPlaybackWidget.h"
 #include "UObject/Object.h"
 #include "UObject/Package.h"
@@ -61,6 +62,8 @@ private:
 FLiveLinkHubPlaybackController::FLiveLinkHubPlaybackController()
 {
 	Client = &IModularFeatures::Get().GetModularFeature<ILiveLinkClient>(ILiveLinkClient::ModularFeatureName);
+
+	OnSourceRemovedHandle = Client->OnLiveLinkSourceRemoved().AddRaw(this, &FLiveLinkHubPlaybackController::OnSourceRemoved);
 	
 	RecordingPlayer = MakeUnique<FLiveLinkUAssetRecordingPlayer>();
 	Playhead = MakeShared<FLiveLinkHubAtomicQualifiedFrameTime, ESPMode::ThreadSafe>();
@@ -77,6 +80,11 @@ FLiveLinkHubPlaybackController::~FLiveLinkHubPlaybackController()
 		Thread->WaitForCompletion();
 		Thread.Reset();
 	}
+
+	if (OnSourceRemovedHandle.IsValid() && Client)
+	{
+		Client->OnLiveLinkSourceRemoved().Remove(OnSourceRemovedHandle);
+	}
 }
 
 TSharedRef<SWidget> FLiveLinkHubPlaybackController::MakePlaybackWidget()
@@ -134,7 +142,7 @@ void FLiveLinkHubPlaybackController::StartPlayback()
 
 void FLiveLinkHubPlaybackController::ResumePlayback()
 {
-	bIsInPlayback = true;
+	bIsPlaying = true;
 	bIsPaused = false;
 	FQualifiedFrameTime CurrentTime = GetCurrentTime();
 
@@ -184,6 +192,14 @@ void FLiveLinkHubPlaybackController::PreparePlayback(ULiveLinkRecording* InLiveL
 		// Save the current state of the sources/subjects in a rollback preset.
 		RollbackPreset->BuildFromClient();
 
+		// This clears out any live streams which might be occurring. They will be restored when exiting playback later.
+		{
+			FLiveLinkClient& LiveLinkClient = IModularFeatures::Get().GetModularFeature<FLiveLinkClient>(ILiveLinkClient::ModularFeatureName);
+
+			LiveLinkClient.RemoveAllSources();
+			LiveLinkClient.Tick();
+		}
+		
 		RecordingToPlay->RecordingPreset->ApplyToClientLatent([this](bool)
 		{
 			bIsReady = true;
@@ -203,7 +219,7 @@ void FLiveLinkHubPlaybackController::BeginPlayback(bool bInReverse)
 	bIsReverse = bInReverse;
 	
 	// Either we are paused and should unpause, or we are toggling forward/reverse play modes.
-	if (bIsPaused || !bIsInPlayback || bReverseChange)
+	if (bIsPaused || !bIsPlaying || bReverseChange)
 	{
 		if (ShouldRestart())
 		{
@@ -223,7 +239,7 @@ void FLiveLinkHubPlaybackController::BeginPlayback(bool bInReverse)
 		
 		ResumePlayback();
 	}
-	else if (bIsInPlayback)
+	else if (bIsPlaying)
 	{
 		PausePlayback();
 	}
@@ -238,7 +254,7 @@ void FLiveLinkHubPlaybackController::RestartPlayback()
 	StopPlayback();
 	StartTimestamp = GetCurrentTime().AsSeconds();
 	RecordingPlayer->RestartPlayback(GetCurrentFrame().Value);
-	bIsInPlayback = true;
+	bIsPlaying = true;
 	bIsReverse = bOldReverse;
 }
 
@@ -249,7 +265,7 @@ void FLiveLinkHubPlaybackController::PausePlayback()
 
 void FLiveLinkHubPlaybackController::StopPlayback()
 {
-	bIsInPlayback = false;
+	bIsPlaying = false;
 
 	// Wait for the playback thread to exit...
 	const uint32 CurrentThreadId = FPlatformTLS::GetCurrentThreadId();
@@ -374,7 +390,7 @@ uint32 FLiveLinkHubPlaybackController::Run()
 		bIsPlaybackWaiting = true;
 		PlaybackEvent->Wait();
 		bIsPlaybackWaiting = false;
-		while (bIsInPlayback)
+		while (bIsPlaying)
 		{
 			if (bIsPaused)
 			{
@@ -418,7 +434,7 @@ uint32 FLiveLinkHubPlaybackController::Run()
 		}
 
 		// If the loop ended because the recording is over.
-		bIsInPlayback = false;
+		bIsPlaying = false;
 		
 		// Trigger the playback finished delegate on the game thread.
 		FSimpleDelegateGraphTask::CreateAndDispatchWhenReady(FSimpleDelegateGraphTask::FDelegate::CreateRaw(this, &FLiveLinkHubPlaybackController::OnPlaybackFinished_Internal), TStatId(), nullptr, ENamedThreads::GameThread);
@@ -432,6 +448,23 @@ void FLiveLinkHubPlaybackController::OnPlaybackFinished_Internal()
 	PlaybackFinishedDelegate.Broadcast();
 }
 
+void FLiveLinkHubPlaybackController::OnSourceRemoved(FGuid Guid)
+{
+	if (RecordingToPlay.IsValid())
+	{
+		// Look for a source that is for this recording and eject. This can occur if the user presses the trash icon
+		// on the playback source while it is in playback.
+		for (const FLiveLinkSourcePreset& Presets : RecordingToPlay->RecordingPreset->GetSourcePresets())
+		{
+			if (Presets.Guid == Guid)
+			{
+				Eject();
+				break;
+			}
+		}
+	}
+}
+
 void FLiveLinkHubPlaybackController::PushSubjectData(const FLiveLinkRecordedFrame& NextFrame, bool bForceSync)
 {
 	// If we're sending static data
diff --git a/Engine/Plugins/Animation/LiveLinkHub/Source/LiveLinkHub/Private/Recording/LiveLinkHubPlaybackController.h b/Engine/Plugins/Animation/LiveLinkHub/Source/LiveLinkHub/Private/Recording/LiveLinkHubPlaybackController.h
--- a/Engine/Plugins/Animation/LiveLinkHub/Source/LiveLinkHub/Private/Recording/LiveLinkHubPlaybackController.h
+++ b/Engine/Plugins/Animation/LiveLinkHub/Source/LiveLinkHub/Private/Recording/LiveLinkHubPlaybackController.h
@@ -89,16 +89,22 @@ public:
 		return bIsReady;
 	}
 	
-	/** Returns whether we've started or are actively playing a recording. */ 
+	/** If a recording is loaded into the controller. */
 	bool IsInPlayback() const
 	{
-		return bIsInPlayback.load();
+		return GetRecording().IsValid();
+	}
+
+	/** Returns whether we've started or are actively playing a recording. */ 
+	bool IsPlaying() const
+	{
+		return bIsPlaying.load();
 	}
 
 	/** If playback is paused. */
 	bool IsPaused() const
 	{
-		return bIsPaused.load() || !IsInPlayback();
+		return bIsPaused.load() || !IsPlaying();
 	}
 
 	/** If playback is playing in reverse. */
@@ -144,6 +150,9 @@ private:
 	/** Handler called when playback is finished on the playback thread. Is responsible for resetting the livelink state to what it was before we started playback. */
 	void OnPlaybackFinished_Internal();
 
+	/** When a source has been removed from Live Link Hub. */
+	void OnSourceRemoved(FGuid Guid);
+	
 	/**
 	 * Send data to the client.
 	 * @param NextFrame The frame to push.
@@ -175,8 +184,8 @@ private:
 	FEventRef PlaybackEvent = FEventRef();
 	/** If the playback thread is waiting. */
 	std::atomic<bool> bIsPlaybackWaiting = false;
-	/** Whether we're doing playback. */
-	std::atomic<bool> bIsInPlayback = false;
+	/** Whether a recording is playing. */
+	std::atomic<bool> bIsPlaying = false;
 	/** Whether we're currently paused. */
 	std::atomic<bool> bIsPaused = false;
 	/** If the recording is playing in reverse. */
@@ -212,4 +221,7 @@ private:
 
 	/** Current framerate of the recording, sampled from the latest frame. */
 	FFrameRate CurrentFrameRate;
+
+	/** Delegate handle for when a source is removed. */
+	FDelegateHandle OnSourceRemovedHandle;
 };
diff --git a/Engine/Plugins/Animation/LiveLinkHub/Source/LiveLinkHub/Private/Recording/SLiveLinkHubRecordingView.h b/Engine/Plugins/Animation/LiveLinkHub/Source/LiveLinkHub/Private/Recording/SLiveLinkHubRecordingView.h
--- a/Engine/Plugins/Animation/LiveLinkHub/Source/LiveLinkHub/Private/Recording/SLiveLinkHubRecordingView.h
+++ b/Engine/Plugins/Animation/LiveLinkHub/Source/LiveLinkHub/Private/Recording/SLiveLinkHubRecordingView.h
@@ -6,6 +6,9 @@
 
 #include "Delegates/Delegate.h"
 #include "Delegates/DelegateCombinations.h"
+#include "LiveLinkHubModule.h"
+#include "LiveLinkHubPlaybackController.h"
+#include "Modules/ModuleManager.h"
 #include "Styling/SlateTypes.h"
 #include "Widgets/Images/SImage.h"
 #include "Widgets/Input/SButton.h"
@@ -45,6 +48,7 @@ public:
 		[
 			SNew(SButton)
 				.OnClicked(this, &SLiveLinkHubRecordingView::OnClickRecordButton)
+				.IsEnabled(this, &SLiveLinkHubRecordingView::CanRecord)
 				[
 					SNew(SHorizontalBox)
 					+ SHorizontalBox::Slot()
@@ -98,6 +102,12 @@ public:
 		return IsRecordingDelegate.Execute();
 	}
 
+	bool CanRecord() const
+	{
+		const FLiveLinkHubModule& LiveLinkHubModule = FModuleManager::Get().GetModuleChecked<FLiveLinkHubModule>("LiveLinkHub");
+		return !LiveLinkHubModule.GetPlaybackController()->IsInPlayback();
+	}
+
 private:
 	/** Delegate used to know if the hub is current recording. */
 	FOnGetIsRecording IsRecordingDelegate;
diff --git a/Engine/Plugins/Animation/LiveLinkHub/Source/LiveLinkHub/Private/UI/Widgets/SLiveLinkHubMainTabView.cpp b/Engine/Plugins/Animation/LiveLinkHub/Source/LiveLinkHub/Private/UI/Widgets/SLiveLinkHubMainTabView.cpp
--- a/Engine/Plugins/Animation/LiveLinkHub/Source/LiveLinkHub/Private/UI/Widgets/SLiveLinkHubMainTabView.cpp
+++ b/Engine/Plugins/Animation/LiveLinkHub/Source/LiveLinkHub/Private/UI/Widgets/SLiveLinkHubMainTabView.cpp
@@ -184,6 +184,11 @@ TSharedRef<SDockTab> SLiveLinkHubMainTabView::SpawnSourcesTab(const FSpawnTabArg
 				.ShowPresetPicker(false)
 				.ShowSettings(false)
 				.CustomHeader(CustomToolbarHeader)
+				.IsEnabled_Lambda([this]()
+				{
+					const FLiveLinkHubModule& LiveLinkHubModule = FModuleManager::Get().GetModuleChecked<FLiveLinkHubModule>("LiveLinkHub");
+					return !LiveLinkHubModule.GetPlaybackController()->IsInPlayback();
+				})
 			]
 			+SVerticalBox::Slot()
 			.FillHeight(1.f)
