LiveLinkHub - Fix crash when double clicking a recording while in playback

Create the rollback preset with a unique name
Make sure the OnSourceRemoved callback doesn't eject the recording during PreparePlayback
Schedule PreparePlayback in the Eject callback to make sure it happens after the livelink preset was applied

diff --git a/Engine/Plugins/Animation/LiveLink/Source/LiveLink/Private/LiveLinkPreset.cpp b/Engine/Plugins/Animation/LiveLink/Source/LiveLink/Private/LiveLinkPreset.cpp
--- a/Engine/Plugins/Animation/LiveLink/Source/LiveLink/Private/LiveLinkPreset.cpp
+++ b/Engine/Plugins/Animation/LiveLink/Source/LiveLink/Private/LiveLinkPreset.cpp
@@ -235,23 +235,29 @@ void ULiveLinkPreset::ApplyToClientLatent(TFunction<void(bool)> CompletionCallba
 
 	ApplyToClientPollingOperation = MakePimpl<FApplyToClientPollingOperation>(this);
 
-	ApplyToClientEndFrameHandle = FCoreDelegates::OnEndFrame.AddLambda([WeakThis = TWeakObjectPtr<ULiveLinkPreset>(this), WrappedCallback = MoveTemp(CompletionCallback)]()
+	
+	ApplyToClientEndFrameHandle = FCoreDelegates::OnEndFrame.AddLambda([WeakThis = TWeakObjectPtr<ULiveLinkPreset>(this), WrappedCallback = MoveTemp(CompletionCallback), Count = ApplyCount]()
 	{
 		if (ensure(WeakThis->ApplyToClientPollingOperation))
 		{
 			FApplyToClientPollingOperation::EApplyToClientUpdateResult Result = WeakThis->ApplyToClientPollingOperation->Update();
 			if (Result != FApplyToClientPollingOperation::EApplyToClientUpdateResult::Pending)
 			{
+				WeakThis->ApplyToClientPollingOperation.Reset();
+				
 				if (WrappedCallback != nullptr)
 				{
 					WrappedCallback(Result == FApplyToClientPollingOperation::EApplyToClientUpdateResult::Success ? true : false);
 				}
 
-				WeakThis->ApplyToClientPollingOperation.Reset();
+				UE_LOG(LogTemp, Verbose, TEXT("Clearing preset %s (%d)"), *WeakThis->GetFullName(), Count);
+
 				WeakThis->ClearApplyToClientTimer();
 			}
 		}
 	});
+
+	UE_LOG(LogTemp, Verbose, TEXT("Applying preset %s (%d)"), *GetFullName(), ApplyCount++);
 }
 
 bool ULiveLinkPreset::AddToClient(const bool bRecreatePresets) const
diff --git a/Engine/Plugins/Animation/LiveLink/Source/LiveLink/Public/LiveLinkPreset.h b/Engine/Plugins/Animation/LiveLink/Source/LiveLink/Public/LiveLinkPreset.h
--- a/Engine/Plugins/Animation/LiveLink/Source/LiveLink/Public/LiveLinkPreset.h
+++ b/Engine/Plugins/Animation/LiveLink/Source/LiveLink/Public/LiveLinkPreset.h
@@ -70,6 +70,9 @@ private:
 
 	/** Holds the current ApplyToClient async operation. Only one operation for all presets can be done at a time. */
 	static TPimplPtr<struct FApplyToClientPollingOperation> ApplyToClientPollingOperation;
+
+	/** Utility variable used to keep track of the number of times this was applied. */
+	uint32 ApplyCount = 0;
 };
 
 #if UE_ENABLE_INCLUDE_ORDER_DEPRECATED_IN_5_2
diff --git a/Engine/Plugins/Animation/LiveLinkHub/Source/LiveLinkHub/Private/Recording/Implementations/LiveLinkUAssetRecordingPlayer.cpp b/Engine/Plugins/Animation/LiveLinkHub/Source/LiveLinkHub/Private/Recording/Implementations/LiveLinkUAssetRecordingPlayer.cpp
--- a/Engine/Plugins/Animation/LiveLinkHub/Source/LiveLinkHub/Private/Recording/Implementations/LiveLinkUAssetRecordingPlayer.cpp
+++ b/Engine/Plugins/Animation/LiveLinkHub/Source/LiveLinkHub/Private/Recording/Implementations/LiveLinkUAssetRecordingPlayer.cpp
@@ -300,7 +300,7 @@ FFrameRate FLiveLinkPlaybackTracks::GetInitialFrameRate() const
 		}
 	}
 
-	UE_LOG(LogLiveLinkHub, Warning, TEXT("Could not find an initial framerate for the recording. Using the defualt value."));
+	UE_LOG(LogLiveLinkHub, Warning, TEXT("Could not find an initial framerate for the recording. Using the default value."));
 	
 	return FFrameRate(30, 1);
 }
diff --git a/Engine/Plugins/Animation/LiveLinkHub/Source/LiveLinkHub/Private/Recording/LiveLinkHubPlaybackController.cpp b/Engine/Plugins/Animation/LiveLinkHub/Source/LiveLinkHub/Private/Recording/LiveLinkHubPlaybackController.cpp
--- a/Engine/Plugins/Animation/LiveLinkHub/Source/LiveLinkHub/Private/Recording/LiveLinkHubPlaybackController.cpp
+++ b/Engine/Plugins/Animation/LiveLinkHub/Source/LiveLinkHub/Private/Recording/LiveLinkHubPlaybackController.cpp
@@ -171,40 +171,56 @@ void FLiveLinkHubPlaybackController::PreparePlayback(ULiveLinkRecording* InLiveL
 	}
 	else if (InLiveLinkRecording != RecordingToPlay.Get())
 	{
-		if (RecordingToPlay.IsValid())
+		TStrongObjectPtr<ULiveLinkRecording> RecordingStrongPtr(InLiveLinkRecording);
+
+		auto PreparePlaybackCallback = [this, RecordingStrongPtr]()
 		{
-			Eject();
-		}
-		
-		RecordingToPlay.Reset(InLiveLinkRecording);
-		RecordingPlayer->PreparePlayback(RecordingToPlay.Get());
+			if (RecordingStrongPtr.IsValid())
+			{
+				// Make sure PreparingPlayback is set to false when the scope exits.
+				TGuardValue<bool> PreparingPlaybackGuard(bIsPreparingPlayback, true);
+
+				RecordingToPlay.Reset(RecordingStrongPtr.Get());
+				RecordingPlayer->PreparePlayback(RecordingToPlay.Get());
 		
-		CurrentFrameRate = RecordingPlayer->GetInitialFramerate();
+				CurrentFrameRate = RecordingPlayer->GetInitialFramerate();
 
-		// The start and end of playback.
-		SetSelectionStartTime(FQualifiedFrameTime(FFrameTime::FromDecimal(0.f), GetFrameRate()));
-		SetSelectionEndTime(GetLength());
+				// The start and end of playback.
+				SetSelectionStartTime(FQualifiedFrameTime(FFrameTime::FromDecimal(0.f), GetFrameRate()));
+				SetSelectionEndTime(GetLength());
 		
-		// The range the user sees.
-		SliderViewRange = TRange<double>(SelectionStartTime.AsSeconds(), SelectionEndTime.AsSeconds());
+				// The range the user sees.
+				SliderViewRange = TRange<double>(SelectionStartTime.AsSeconds(), SelectionEndTime.AsSeconds());
 		
-		RollbackPreset.Reset(NewObject<ULiveLinkPreset>(GetTransientPackage(), TEXT("RecordingRollbackPreset")));
-		// Save the current state of the sources/subjects in a rollback preset.
-		RollbackPreset->BuildFromClient();
+				const FName UniqueName = MakeUniqueObjectName(GetTransientPackage(), ULiveLinkPreset::StaticClass(), TEXT("RecordingRollbackPreset"));
+				RollbackPreset.Reset(NewObject<ULiveLinkPreset>(GetTransientPackage(), UniqueName));
+				// Save the current state of the sources/subjects in a rollback preset.
+				RollbackPreset->BuildFromClient();
 
-		// This clears out any live streams which might be occurring. They will be restored when exiting playback later.
-		{
-			FLiveLinkClient& LiveLinkClient = IModularFeatures::Get().GetModularFeature<FLiveLinkClient>(ILiveLinkClient::ModularFeatureName);
+				// This clears out any live streams which might be occurring. They will be restored when exiting playback later.
+				{
+					FLiveLinkClient& LiveLinkClient = IModularFeatures::Get().GetModularFeature<FLiveLinkClient>(ILiveLinkClient::ModularFeatureName);
 
-			LiveLinkClient.RemoveAllSources();
-			LiveLinkClient.Tick();
-		}
+					LiveLinkClient.RemoveAllSources();
+					LiveLinkClient.Tick();
+				}
 		
-		RecordingToPlay->RecordingPreset->ApplyToClientLatent([this](bool)
+				RecordingToPlay->RecordingPreset->ApplyToClientLatent([this](bool)
+				{
+					bIsReady = true;
+					SyncToFrame(0); // Needed to establish connection with client
+				});
+			}
+		};
+
+		if (RecordingToPlay.IsValid())
 		{
-			bIsReady = true;
-			SyncToFrame(0); // Needed to establish connection with client
-		});
+			Eject(PreparePlaybackCallback);
+		}
+		else
+		{
+			PreparePlaybackCallback();
+		}
 	}
 }
 
@@ -283,7 +299,7 @@ void FLiveLinkHubPlaybackController::StopPlayback()
 	bIsReverse = false;
 }
 
-void FLiveLinkHubPlaybackController::Eject()
+void FLiveLinkHubPlaybackController::Eject(TFunction<void()> CompletionCallback)
 {
 	bIsReady = false;
 	
@@ -296,13 +312,24 @@ void FLiveLinkHubPlaybackController::Eject()
 	SetSelectionEndTime(FQualifiedFrameTime(FFrameTime::FromDecimal(0), GetFrameRate()));
 	Playhead->SetValue(FQualifiedFrameTime(FFrameTime::FromDecimal(0), GetFrameRate()));
 	StartTimestamp = 0.f;
+
+	// Recording is done, clear the pointer.
+	RecordingToPlay.Reset();
 	
 	if (RollbackPreset.IsValid())
 	{
-		RollbackPreset->ApplyToClientLatent();
+		RollbackPreset->ApplyToClientLatent([CompletionCallback](bool)
+		{
+			if (CompletionCallback)
+			{
+				CompletionCallback();
+			}
+		});
+	}
+	else if (CompletionCallback)
+	{
+		CompletionCallback();
 	}
-	// Recording is done, clear the pointer.
-	RecordingToPlay.Reset();
 }
 
 void FLiveLinkHubPlaybackController::GoToTime(FQualifiedFrameTime InTime)
@@ -450,7 +477,7 @@ void FLiveLinkHubPlaybackController::OnPlaybackFinished_Internal()
 
 void FLiveLinkHubPlaybackController::OnSourceRemoved(FGuid Guid)
 {
-	if (RecordingToPlay.IsValid())
+	if (RecordingToPlay.IsValid() && !bIsPreparingPlayback)
 	{
 		// Look for a source that is for this recording and eject. This can occur if the user presses the trash icon
 		// on the playback source while it is in playback.
diff --git a/Engine/Plugins/Animation/LiveLinkHub/Source/LiveLinkHub/Private/Recording/LiveLinkHubPlaybackController.h b/Engine/Plugins/Animation/LiveLinkHub/Source/LiveLinkHub/Private/Recording/LiveLinkHubPlaybackController.h
--- a/Engine/Plugins/Animation/LiveLinkHub/Source/LiveLinkHub/Private/Recording/LiveLinkHubPlaybackController.h
+++ b/Engine/Plugins/Animation/LiveLinkHub/Source/LiveLinkHub/Private/Recording/LiveLinkHubPlaybackController.h
@@ -54,7 +54,7 @@ public:
 	void StopPlayback();
 
 	/** Stop playback and restore the previous settings. */
-	void Eject();
+	void Eject(TFunction<void()> CompletionCallback = nullptr);
 
 	/** Go to a specific time. */
 	void GoToTime(FQualifiedFrameTime InTime);
@@ -192,6 +192,8 @@ private:
 	std::atomic<bool>bIsReverse = false;
 	/** The timestamp of the animation when first playing. Can be > 0 when running in reverse. */
 	std::atomic<double> StartTimestamp = 0.f;
+	/** Indicates that we're in the process of preparing the playback. Used by the OnSourceRemoved callback to make sure we don't eject during the PreparePlayback step. */
+	bool bIsPreparingPlayback = false;
 	/** LiveLinkRecording to play.  */
 	TStrongObjectPtr<ULiveLinkRecording> RecordingToPlay;
 	/** Delegate called when a recording playback is finished (if it's not looping). */
