LiveLinkHub - Add missing context menu for removing client from clients view

diff --git a/Engine/Plugins/Animation/LiveLinkHub/Source/LiveLinkHub/Private/Clients/SLiveLinkHubClientsView.h b/Engine/Plugins/Animation/LiveLinkHub/Source/LiveLinkHub/Private/Clients/SLiveLinkHubClientsView.h
--- a/Engine/Plugins/Animation/LiveLinkHub/Source/LiveLinkHub/Private/Clients/SLiveLinkHubClientsView.h
+++ b/Engine/Plugins/Animation/LiveLinkHub/Source/LiveLinkHub/Private/Clients/SLiveLinkHubClientsView.h
@@ -4,6 +4,8 @@
 
 #include "Widgets/SCompoundWidget.h"
 
+#include "Algo/ForEach.h"
+#include "Algo/Transform.h"
 #include "Async/Async.h"
 #include "Delegates/Delegate.h"
 #include "Delegates/DelegateCombinations.h"
@@ -368,6 +370,7 @@ public:
 				.ItemHeight(20.0f)
 				.OnSelectionChanged(this, &SLiveLinkHubClientsView::OnSelectionChanged)
 				.OnGenerateRow(this, &SLiveLinkHubClientsView::OnGenerateClientRow)
+				.OnContextMenuOpening(this, &SLiveLinkHubClientsView::OnContextMenuOpening)
 				.OnGetChildren(this, &SLiveLinkHubClientsView::OnGetChildren)
 				.OnKeyDownHandler(this, &SLiveLinkHubClientsView::OnKeyDownHandler)
 				.HeaderRow
@@ -464,6 +467,44 @@ private:
 			.Item(Item);
 	}
 
+	/** Handler used to create the context menu widget. */
+	TSharedPtr<SWidget> OnContextMenuOpening() const
+	{
+		constexpr bool CloseAfterSelection = true;
+		FMenuBuilder MenuBuilder( CloseAfterSelection, nullptr);
+		MenuBuilder.AddMenuEntry(
+			LOCTEXT("Remove", "Remove selected client"),
+			LOCTEXT("RemoveClientTooltip", "Stop transmitting LiveLink data to this client."),
+			FSlateIcon("LiveLinkStyle", "LiveLinkClient.Common.RemoveSource"),
+			FUIAction(
+				FExecuteAction::CreateRaw(this, &SLiveLinkHubClientsView::RemoveSelectedClient),
+				FCanExecuteAction::CreateStatic( [](){ return true; } )
+			)
+		);
+
+		MenuBuilder.AddMenuEntry(
+			LOCTEXT("RemoveAll", "Remove all clients"),
+			LOCTEXT("RemoveAllClientTooltip", "Stop transmitting LiveLink data to all discovered clients."),
+			FSlateIcon("LiveLinkStyle", "LiveLinkClient.Common.RemoveSource"),
+			FUIAction(
+				FExecuteAction::CreateRaw(this, &SLiveLinkHubClientsView::RemoveAllClients),
+				FCanExecuteAction::CreateStatic([]() { return true; })
+			)
+		);
+
+
+		return MenuBuilder.MakeWidget();
+	}
+
+	/** Handler called when selection changes in the list view. */
+	void OnSelectionChanged(FClientTreeItemPtr InItem, const ESelectInfo::Type InSelectInfoType) const
+	{
+		if (InItem)
+		{
+			OnClientSelectedDelegate.ExecuteIfBound(InItem->ClientId);
+		}
+	}
+
 	/** Handler called to fetch a tree row's children. */
 	void OnGetChildren(FClientTreeItemPtr Item, TArray<FClientTreeItemPtr>& OutChildren)
 	{
@@ -471,7 +512,7 @@ private:
 	}
 
 	/** Method to handle deleting clients when the delete key is pressed. */
-    FReply OnKeyDownHandler(const FGeometry&, const FKeyEvent& InKeyEvent)
+    FReply OnKeyDownHandler(const FGeometry&, const FKeyEvent& InKeyEvent) const
     {
 		if (InKeyEvent.GetKey() == EKeys::Delete || InKeyEvent.GetKey() == EKeys::BackSpace)
     	{
@@ -485,15 +526,30 @@ private:
     	return FReply::Unhandled();
     }
 
-	/** Handler called when selection changes in the list view. */
-	void OnSelectionChanged(FClientTreeItemPtr InItem, const ESelectInfo::Type InSelectInfoType)
+	/** Remove the selected client from the list. This will stop all livelink messages from being transmitted to it. */
+	void RemoveSelectedClient() const
 	{
-		if (InItem)
+		if (TOptional<FLiveLinkHubClientId> Id = GetSelectedClient())
 		{
-			OnClientSelectedDelegate.ExecuteIfBound(InItem->ClientId);
+			OnRemoveClientFromSessionDelegate.ExecuteIfBound(*Id);
 		}
 	}
 
+	/** Remove all clients from the list. This will stop all livelink messages from being transmitted to them. */
+	void RemoveAllClients() const
+	{
+		TArray<FLiveLinkHubClientId> ClientIds;
+		ClientIds.Reserve(Clients.Num());
+
+		Algo::TransformIf(Clients, ClientIds, [](FClientTreeItemPtr ClientPtr) { return !!ClientPtr; }, [](FClientTreeItemPtr ClientPtr) { return ClientPtr->ClientId; });
+		
+		Algo::ForEach(ClientIds, [this](const FLiveLinkHubClientId& ClientId)
+		{
+			OnRemoveClientFromSessionDelegate.ExecuteIfBound(ClientId);
+		});
+	}
+
+
 	/** Handler called when a client is picked in the Add Client menu. */
 	void OnDiscoveredClientPicked(TSharedPtr<FLiveLinkHubClientId> InItem, const ESelectInfo::Type InSelectInfoType)
 	{
