LiveLinkHub - Fix timecode setting being sent out to clients that weren't enabled

diff --git a/Engine/Plugins/Animation/LiveLinkHub/Source/LiveLinkHub/Private/Clients/LiveLinkHubProvider.cpp b/Engine/Plugins/Animation/LiveLinkHub/Source/LiveLinkHub/Private/Clients/LiveLinkHubProvider.cpp
--- a/Engine/Plugins/Animation/LiveLinkHub/Source/LiveLinkHub/Private/Clients/LiveLinkHubProvider.cpp
+++ b/Engine/Plugins/Animation/LiveLinkHub/Source/LiveLinkHub/Private/Clients/LiveLinkHubProvider.cpp
@@ -54,35 +54,12 @@ FLiveLinkHubProvider::~FLiveLinkHubProvider()
 
 bool FLiveLinkHubProvider::ShouldTransmitToSubject_AnyThread(FName SubjectName, FMessageAddress Address) const
 {
-	FReadScopeLock Locker(ClientsMapLock);
-
-	const FLiveLinkHubClientId ClientId = AddressToIdCache.FindRef(Address);
-	if (const FLiveLinkHubUEClientInfo* ClientInfoPtr = ClientsMap.Find(ClientId))
+	auto AdditionalFilter = [SubjectName](const FLiveLinkHubUEClientInfo* ClientInfoPtr)
 	{
-		if (const TSharedPtr<ILiveLinkHubSessionManager> Manager = SessionManager.Pin())
-		{
-			if (const TSharedPtr<ILiveLinkHubSession> CurrentSession = Manager->GetCurrentSession())
-			{
-				if (!CurrentSession->IsClientInSession(ClientInfoPtr->Id))
-				{
-					return false;
-				}
-			}
-		}
-
-		if (!ClientInfoPtr->bEnabled)
-		{
-			return false;
-		}
-
 		return !ClientInfoPtr->DisabledSubjects.Contains(SubjectName);
-	}
-	else
-	{
-		UE_LOG(LogLiveLinkHub, Warning, TEXT("Attempted to transmit data to an invalid client."));
-	}
+	};
 
-	return true;
+	return ShouldTransmitToClient_AnyThread(Address, AdditionalFilter);
 }
 
 void FLiveLinkHubProvider::SetTimecodeSettings(FLiveLinkHubTimecodeSettings InSettings)
@@ -93,7 +70,7 @@ void FLiveLinkHubProvider::SetTimecodeSettings(FLiveLinkHubTimecodeSettings InSe
 
 void FLiveLinkHubProvider::SendTimecodeSettings()
 {
-	SendMessage(FMessageEndpoint::MakeMessage<FLiveLinkHubTimecodeSettings>(TimecodeSettings));
+	SendMessageToEnabledClients(FMessageEndpoint::MakeMessage<FLiveLinkHubTimecodeSettings>(TimecodeSettings));
 }
 
 void FLiveLinkHubProvider::AddRestoredClient(FLiveLinkHubUEClientInfo& RestoredClientInfo)
@@ -250,6 +227,44 @@ void FLiveLinkHubProvider::HandleClientInfoMessage(const FLiveLinkClientInfoMess
 	}
 }
 
+bool FLiveLinkHubProvider::ShouldTransmitToClient_AnyThread(FMessageAddress Address, TFunctionRef<bool(const FLiveLinkHubUEClientInfo* ClientInfoPtr)> AdditionalFilter) const
+{
+	if (!Address.IsValid())
+	{
+		return false;
+	}
+
+	FReadScopeLock Locker(ClientsMapLock);
+
+	const FLiveLinkHubClientId ClientId = AddressToIdCache.FindRef(Address);
+	if (const FLiveLinkHubUEClientInfo* ClientInfoPtr = ClientsMap.Find(ClientId))
+	{
+		if (const TSharedPtr<ILiveLinkHubSessionManager> Manager = SessionManager.Pin())
+		{
+			if (const TSharedPtr<ILiveLinkHubSession> CurrentSession = Manager->GetCurrentSession())
+			{
+				if (!CurrentSession->IsClientInSession(ClientInfoPtr->Id))
+				{
+					return false;
+				}
+			}
+		}
+
+		if (!ClientInfoPtr->bEnabled)
+		{
+			return false;
+		}
+
+		return AdditionalFilter(ClientInfoPtr);
+	}
+	else
+	{
+		UE_LOG(LogLiveLinkHub, Warning, TEXT("Attempted to transmit data to an invalid client."));
+	}
+
+	return true;
+}
+
 void FLiveLinkHubProvider::OnConnectionsClosed(const TArray<FMessageAddress>& ClosedAddresses)
 {
 	// List of OldId -> NewId
diff --git a/Engine/Plugins/Animation/LiveLinkHub/Source/LiveLinkHub/Private/Clients/LiveLinkHubProvider.h b/Engine/Plugins/Animation/LiveLinkHub/Source/LiveLinkHub/Private/Clients/LiveLinkHubProvider.h
--- a/Engine/Plugins/Animation/LiveLinkHub/Source/LiveLinkHub/Private/Clients/LiveLinkHubProvider.h
+++ b/Engine/Plugins/Animation/LiveLinkHub/Source/LiveLinkHub/Private/Clients/LiveLinkHubProvider.h
@@ -9,6 +9,8 @@
 #include "IMessageContext.h"
 #include "LiveLinkHubMessages.h"
 #include "LiveLinkProviderImpl.h"
+#include "Templates/Function.h"
+
 
 class ILiveLinkHubSessionManager;
 struct FLiveLinkHubClientId;
@@ -65,6 +67,27 @@ private:
 	/** Send timecode settings to connected Live Link Hub provider. */
 	void SendTimecodeSettings();
 
+	/** Send a message to clients that are connected and enabled through the hub clients list. */
+	template<typename MessageType>
+	void SendMessageToEnabledClients(MessageType* Message)
+	{
+		TArray<FMessageAddress> AllAddresses;
+		GetConnectedAddresses(AllAddresses);
+
+		TArray<FMessageAddress> EnabledAddresses = AllAddresses.FilterByPredicate([this](const FMessageAddress& Address)
+		{
+			return ShouldTransmitToClient_AnyThread(Address);
+		});
+
+		SendMessage(Message, EnabledAddresses);
+	}
+
+	/**
+	 * Whether a message should be transmitted to a particular client, identified by a message address.
+	 * You can specify an additional filter method if you want to filter based on the client info.
+	 **/
+	bool ShouldTransmitToClient_AnyThread(FMessageAddress Address, TFunctionRef<bool(const FLiveLinkHubUEClientInfo* ClientInfoPtr)> AdditionalFilter = [](const FLiveLinkHubUEClientInfo* ClientInfoPtr){ return true; }) const;
+
 protected:
 	//~ Begin ILiveLinkHubClientsModel interface
 	virtual void OnConnectionsClosed(const TArray<FMessageAddress>& ClosedAddresses) override;
diff --git a/Engine/Source/Runtime/LiveLinkMessageBusFramework/Public/LiveLinkProviderImpl.h b/Engine/Source/Runtime/LiveLinkMessageBusFramework/Public/LiveLinkProviderImpl.h
--- a/Engine/Source/Runtime/LiveLinkMessageBusFramework/Public/LiveLinkProviderImpl.h
+++ b/Engine/Source/Runtime/LiveLinkMessageBusFramework/Public/LiveLinkProviderImpl.h
@@ -126,6 +126,17 @@ protected:
 		MessageEndpoint->Send(Message, Address);
 	}
 
+	template<typename MessageType>
+	void SendMessage(MessageType* Message, const TArray<FMessageAddress>& Addresses)
+	{
+		if (!Message || !Addresses.Num())
+		{
+			return;
+		}
+
+		MessageEndpoint->Send(Message, Addresses);
+	}
+
 	template<typename MessageType>
 	void Subscribe()
 	{
