LiveLinkhub - Send updated timecode settings when client is enabled

diff --git a/Engine/Plugins/Animation/LiveLinkHub/Source/LiveLinkHub/Private/Clients/LiveLinkHubProvider.cpp b/Engine/Plugins/Animation/LiveLinkHub/Source/LiveLinkHub/Private/Clients/LiveLinkHubProvider.cpp
--- a/Engine/Plugins/Animation/LiveLinkHub/Source/LiveLinkHub/Private/Clients/LiveLinkHubProvider.cpp
+++ b/Engine/Plugins/Animation/LiveLinkHub/Source/LiveLinkHub/Private/Clients/LiveLinkHubProvider.cpp
@@ -389,10 +389,18 @@ bool FLiveLinkHubProvider::IsClientConnected(FLiveLinkHubClientId Client) const
 
 void FLiveLinkHubProvider::SetClientEnabled(FLiveLinkHubClientId Client, bool bInEnable)
 {
-	FWriteScopeLock Locker(ClientsMapLock);
-	if (FLiveLinkHubUEClientInfo* ClientInfoPtr = ClientsMap.Find(Client))
 	{
-		ClientInfoPtr->bEnabled = bInEnable;
+		FWriteScopeLock Locker(ClientsMapLock);
+		if (FLiveLinkHubUEClientInfo* ClientInfoPtr = ClientsMap.Find(Client))
+		{
+			ClientInfoPtr->bEnabled = bInEnable;
+		}
+	}
+
+	// No need to update timecode settings if client was disabled.
+	if (bInEnable)
+	{
+		SendTimecodeSettings();
 	}
 }
 
