LiveLinkHub - Add entry point to register main tabs through plugins

diff --git a/Engine/Plugins/Animation/LiveLinkHub/Source/LiveLinkHub/Private/ILiveLinkHubComponent.cpp b/Engine/Plugins/Animation/LiveLinkHub/Source/LiveLinkHub/Private/ILiveLinkHubComponent.cpp
new file mode 100644
--- /dev/null
+++ b/Engine/Plugins/Animation/LiveLinkHub/Source/LiveLinkHub/Private/ILiveLinkHubComponent.cpp
@@ -0,0 +1,5 @@
+// Copyright Epic Games, Inc. All Rights Reserved.
+
+#include "ILiveLinkHubComponent.h"
+
+const FName ILiveLinkHubComponent::ModularFeatureName = "LiveLinkHubComponent";
diff --git a/Engine/Plugins/Animation/LiveLinkHub/Source/LiveLinkHub/Private/UI/Widgets/LiveLinkHubMainTabController.cpp b/Engine/Plugins/Animation/LiveLinkHub/Source/LiveLinkHub/Private/UI/Widgets/LiveLinkHubMainTabController.cpp
--- a/Engine/Plugins/Animation/LiveLinkHub/Source/LiveLinkHub/Private/UI/Widgets/LiveLinkHubMainTabController.cpp
+++ b/Engine/Plugins/Animation/LiveLinkHub/Source/LiveLinkHub/Private/UI/Widgets/LiveLinkHubMainTabController.cpp
@@ -12,7 +12,7 @@ static const FName LiveLinkHubTabID = "LiveLinkHubMainTab";
 
 void FLiveLinkHubMainTabController::Init(const FLiveLinkHubComponentInitParams& Params)
 {
-	FGlobalTabmanager::Get()->RegisterTabSpawner(LiveLinkHubTabID, FOnSpawnTab::CreateRaw(this, &FLiveLinkHubMainTabController::SpawnMainTab, Params.WindowController->GetRootWindow()))
+	FGlobalTabmanager::Get()->RegisterTabSpawner(LiveLinkHubTabID, FOnSpawnTab::CreateRaw(this, &FLiveLinkHubMainTabController::SpawnMainTab, Params.Window))
 		.SetDisplayName(LOCTEXT("LiveLinkHubTabLabel", "LiveLink Hub"))
 		.SetIcon(FSlateIcon("LiveLinkStyle", "LiveLinkHub.Icon.Small"))
 		.SetTooltipText(LOCTEXT("MainTabTooltipText", "LiveLink Hub Main Tab"));
@@ -25,7 +25,7 @@ void FLiveLinkHubMainTabController::OpenTab()
 	FGlobalTabmanager::Get()->TryInvokeTab(LiveLinkHubTabID);
 }
 
-TSharedRef<SDockTab> FLiveLinkHubMainTabController::SpawnMainTab(const FSpawnTabArgs& Args, TSharedPtr<SWindow> RootWindow)
+TSharedRef<SDockTab> FLiveLinkHubMainTabController::SpawnMainTab(const FSpawnTabArgs& Args, TSharedRef<SWindow> RootWindow)
 {
 	const TSharedRef<SDockTab> DockTab = SNew(SDockTab)
 		.Label(LOCTEXT("LiveLinkHubTitle", "LiveLink Hub"))
diff --git a/Engine/Plugins/Animation/LiveLinkHub/Source/LiveLinkHub/Private/UI/Widgets/LiveLinkHubMainTabController.h b/Engine/Plugins/Animation/LiveLinkHub/Source/LiveLinkHub/Private/UI/Widgets/LiveLinkHubMainTabController.h
--- a/Engine/Plugins/Animation/LiveLinkHub/Source/LiveLinkHub/Private/UI/Widgets/LiveLinkHubMainTabController.h
+++ b/Engine/Plugins/Animation/LiveLinkHub/Source/LiveLinkHub/Private/UI/Widgets/LiveLinkHubMainTabController.h
@@ -2,7 +2,7 @@
 
 #pragma once
 
-#include "UI/Window/ILiveLinkHubComponent.h"
+#include "ILiveLinkHubComponent.h"
 
 class FSpawnTabArgs;
 class SLiveLinkHubMainTabView;
@@ -11,7 +11,6 @@ class SWindow;
 class FLiveLinkHubMainTabController : public ILiveLinkHubComponent, public TSharedFromThis<FLiveLinkHubMainTabController>
 {
 public:
-
 	//~ Begin ILiveLinkHubComponent Interface
 	virtual void Init(const FLiveLinkHubComponentInitParams& Params) override;
 	//~ End ILiveLinkHubComponent Interface
@@ -24,5 +23,5 @@ private:
 	TSharedPtr<SLiveLinkHubMainTabView> MainTabView;
 
 	/** Spawn the main tab. */
-	TSharedRef<SDockTab> SpawnMainTab(const FSpawnTabArgs& Args, TSharedPtr<SWindow> RootWindow);
+	TSharedRef<SDockTab> SpawnMainTab(const FSpawnTabArgs& Args, TSharedRef<SWindow> RootWindow);
 };
diff --git a/Engine/Plugins/Animation/LiveLinkHub/Source/LiveLinkHub/Private/UI/Window/ILiveLinkHubComponent.h b/Engine/Plugins/Animation/LiveLinkHub/Source/LiveLinkHub/Private/UI/Window/ILiveLinkHubComponent.h
deleted file mode 100644
diff --git a/Engine/Plugins/Animation/LiveLinkHub/Source/LiveLinkHub/Private/UI/Window/LiveLinkHubWindowController.cpp b/Engine/Plugins/Animation/LiveLinkHub/Source/LiveLinkHub/Private/UI/Window/LiveLinkHubWindowController.cpp
--- a/Engine/Plugins/Animation/LiveLinkHub/Source/LiveLinkHub/Private/UI/Window/LiveLinkHubWindowController.cpp
+++ b/Engine/Plugins/Animation/LiveLinkHub/Source/LiveLinkHub/Private/UI/Window/LiveLinkHubWindowController.cpp
@@ -3,6 +3,7 @@
 #include "LiveLinkHubWindowController.h"
 
 #include "CoreGlobals.h"
+#include "Features/IModularFeatures.h"
 #include "Framework/Application/SlateApplication.h"
 #include "Framework/Docking/LayoutService.h"
 #include "Framework/Notifications/NotificationManager.h"
@@ -18,8 +19,6 @@ FLiveLinkHubWindowController::FLiveLinkHubWindowController(const FLiveLinkHubWin
 	: LiveLinkHubLayoutIni(Params.LiveLinkHubLayoutIni)
 	, MainTabController(MakeShared<FLiveLinkHubMainTabController>())
 {
-	LiveLinkHubComponents.Add(MainTabController);
-	
 	ModalWindowManager = InitializeSlateApplication();
 }
 
@@ -115,8 +114,11 @@ TSharedPtr<FModalWindowManager> FLiveLinkHubWindowController::InitializeSlateApp
 
 void FLiveLinkHubWindowController::InitComponents(const TSharedRef<FTabManager::FStack>& MainArea)
 {
-	const FLiveLinkHubComponentInitParams Params { SharedThis(this), MainArea };
-	for (const TSharedRef<ILiveLinkHubComponent>& LiveLinkHubComponent : LiveLinkHubComponents)
+	const FLiveLinkHubComponentInitParams Params { GetRootWindow().ToSharedRef(), MainArea };
+	MainTabController->Init(Params);
+
+	TArray<ILiveLinkHubComponent*> ExternalComponents = IModularFeatures::Get().GetModularFeatureImplementations<ILiveLinkHubComponent>(ILiveLinkHubComponent::ModularFeatureName);
+	for (ILiveLinkHubComponent* LiveLinkHubComponent : ExternalComponents)
 	{
 		LiveLinkHubComponent->Init(Params);
 	}
diff --git a/Engine/Plugins/Animation/LiveLinkHub/Source/LiveLinkHub/Private/UI/Window/LiveLinkHubWindowController.h b/Engine/Plugins/Animation/LiveLinkHub/Source/LiveLinkHub/Private/UI/Window/LiveLinkHubWindowController.h
--- a/Engine/Plugins/Animation/LiveLinkHub/Source/LiveLinkHub/Private/UI/Window/LiveLinkHubWindowController.h
+++ b/Engine/Plugins/Animation/LiveLinkHub/Source/LiveLinkHub/Private/UI/Window/LiveLinkHubWindowController.h
@@ -58,8 +58,6 @@ private:
 	TSharedPtr<SWindow> RootWindow;
 	/** Manages the main tab. */
 	TSharedRef<FLiveLinkHubMainTabController> MainTabController;
-	/** LiveLinkHub components */
-	TArray<TSharedRef<ILiveLinkHubComponent>> LiveLinkHubComponents;
 	/** Manages modal windows for the application. */
 	TSharedPtr<FModalWindowManager> ModalWindowManager;
 	/** Menu bar widget for the hub. */
diff --git a/Engine/Plugins/Animation/LiveLinkHub/Source/LiveLinkHub/Public/ILiveLinkHubComponent.h b/Engine/Plugins/Animation/LiveLinkHub/Source/LiveLinkHub/Public/ILiveLinkHubComponent.h
new file mode 100644
--- /dev/null
+++ b/Engine/Plugins/Animation/LiveLinkHub/Source/LiveLinkHub/Public/ILiveLinkHubComponent.h
@@ -0,0 +1,37 @@
+﻿// Copyright Epic Games, Inc. All Rights Reserved.
+
+#pragma once
+
+#include "Features/IModularFeature.h"
+#include "Framework/Docking/TabManager.h"
+#include "Templates/SharedPointer.h"
+
+class SWindow;
+
+struct LIVELINKHUB_API FLiveLinkHubComponentInitParams
+{
+	/** LiveLinkHub's main Window. */
+	TSharedRef<SWindow> Window;
+
+	/** The root area of the main window layout. Add tabs to this. */
+	TSharedRef<FTabManager::FStack> MainStack;
+
+	FLiveLinkHubComponentInitParams(TSharedRef<SWindow> InWindow, TSharedRef<FTabManager::FStack> InMainStack)
+		: Window(MoveTemp(InWindow))
+		, MainStack(MoveTemp(InMainStack))
+	{}
+};
+
+
+/** Provides an interface for elements that wish to add additional content (tabs) to LiveLinkHub. */
+class LIVELINKHUB_API ILiveLinkHubComponent : public IModularFeature
+{
+public:
+	/** Name of the modular feature. */
+	static const FName ModularFeatureName;
+
+	virtual ~ILiveLinkHubComponent() = default;
+
+	/** Initialises the component, e.g. registering tab spawners, etc. */
+	virtual void Init(const FLiveLinkHubComponentInitParams& Params) = 0;
+};
