LiveLinkHub - Add support for frame preprocessors and translators

Since we do not evaluate livelink data in the hub, translators behave differently than they do in UE. In the hub, you are limited to a single translator that will always be used. This differs from UE where you can have a list of translators that can be used to translate data when needed. Since LiveLinkHub cannot infer the target subject role, we have to pick it explicitly, hence the limit of 1 translator.

diff --git a/Engine/Plugins/Animation/LiveLink/Source/LiveLink/Private/LiveLinkSubject.cpp b/Engine/Plugins/Animation/LiveLink/Source/LiveLink/Private/LiveLinkSubject.cpp
--- a/Engine/Plugins/Animation/LiveLink/Source/LiveLink/Private/LiveLinkSubject.cpp
+++ b/Engine/Plugins/Animation/LiveLink/Source/LiveLink/Private/LiveLinkSubject.cpp
@@ -335,32 +335,12 @@ bool FLiveLinkSubject::HasStaticData() const
 void FLiveLinkSubject::AddFrameData(FLiveLinkFrameDataStruct&& InFrameData)
 {
 	check(IsInGameThread());
-	if (!StaticData.IsValid())
-	{
-		static const FName InvalidStatFrame = "LiveLinkSubject_InvalidStatFrame";
-		FLiveLinkLog::WarningOnce(InvalidStatFrame, SubjectKey, TEXT("Can't add frame for subject '%s'. The static frame data is invalid."), *SubjectKey.SubjectName.ToString());
-		return;
-	}
 
-	if (Role == nullptr)
+	if (!ValidateFrameData(InFrameData))
 	{
 		return;
 	}
 
-	if (Role->GetDefaultObject<ULiveLinkRole>()->GetFrameDataStruct() != InFrameData.GetStruct())
-	{
-		static const FName NAME_IncompatibleRoles = "LiveLinkSubject_IncompatibleRoles";
-		FLiveLinkLog::WarningOnce(NAME_IncompatibleRoles, SubjectKey, TEXT("Can't add frame for subject '%s'. The frame data is incompatible with current role '%s'."), *SubjectKey.SubjectName.ToString(), *Role->GetName());
-		return;
-	}
-
-	if (!FLiveLinkRoleTrait::Validate(Role, InFrameData))
-	{
-		static const FName NAME_UnsupportedFrameData = "LiveLinkSubject_UnsupportedFrameData";
-		FLiveLinkLog::WarningOnce(NAME_UnsupportedFrameData, SubjectKey, TEXT("Trying to add unsupported frame data type to role '%s'."), *Role->GetName());
-		return;
-	}
-
 	// Before adding the new frame, test to see if we are going to increase the buffer size
 	const bool bRemoveFrame = FrameData.Num() >= CachedSettings.BufferSettings.MaxNumberOfFrameToBuffered;
 	if (bRemoveFrame)
@@ -454,6 +434,45 @@ void FLiveLinkSubject::AddFrameData(FLiveLinkFrameDataStruct&& InFrameData)
 	LastPushTime = FApp::GetCurrentTime();
 }
 
+void FLiveLinkSubject::PreprocessFrame(FLiveLinkFrameDataStruct& InOutFrameData)
+{
+	for (ULiveLinkFramePreProcessor::FWorkerSharedPtr PreProcessor : FramePreProcessors)
+	{
+		PreProcessor->PreProcessFrame(InOutFrameData);
+	}
+}
+
+bool FLiveLinkSubject::ValidateFrameData(const FLiveLinkFrameDataStruct& InFrameData)
+{
+	if (!StaticData.IsValid())
+	{
+		static const FName InvalidStatFrame = "LiveLinkSubject_InvalidStatFrame";
+		FLiveLinkLog::WarningOnce(InvalidStatFrame, SubjectKey, TEXT("Can't add frame for subject '%s'. The static frame data is invalid."), *SubjectKey.SubjectName.ToString());
+		return false;
+	}
+
+	if (Role == nullptr)
+	{
+		return false;
+	}
+
+	if (Role->GetDefaultObject<ULiveLinkRole>()->GetFrameDataStruct() != InFrameData.GetStruct())
+	{
+		static const FName NAME_IncompatibleRoles = "LiveLinkSubject_IncompatibleRoles";
+		FLiveLinkLog::WarningOnce(NAME_IncompatibleRoles, SubjectKey, TEXT("Can't add frame for subject '%s'. The frame data is incompatible with current role '%s'."), *SubjectKey.SubjectName.ToString(), *Role->GetName());
+		return false;
+	}
+
+	if (!FLiveLinkRoleTrait::Validate(Role, InFrameData))
+	{
+		static const FName NAME_UnsupportedFrameData = "LiveLinkSubject_UnsupportedFrameData";
+		FLiveLinkLog::WarningOnce(NAME_UnsupportedFrameData, SubjectKey, TEXT("Trying to add unsupported frame data type to role '%s'."), *Role->GetName());
+		return false;
+	}
+
+	return true;
+}
+
 int32 FLiveLinkSubject::FindNewFrame_WorldTime(const FLiveLinkWorldTime& WorldTime) const
 {
 	if (CachedSettings.BufferSettings.bValidEngineTimeEnabled)
@@ -1262,4 +1281,3 @@ void FLiveLinkSubject::ResetBufferStats()
 	EvaluationStatistics.FrameDrop = 0;
 	EvaluationStatistics.LastEvaluationData = FTimedDataInputEvaluationData();
 }
-
diff --git a/Engine/Plugins/Animation/LiveLink/Source/LiveLink/Public/LiveLinkSubject.h b/Engine/Plugins/Animation/LiveLink/Source/LiveLink/Public/LiveLinkSubject.h
--- a/Engine/Plugins/Animation/LiveLink/Source/LiveLink/Public/LiveLinkSubject.h
+++ b/Engine/Plugins/Animation/LiveLink/Source/LiveLink/Public/LiveLinkSubject.h
@@ -111,6 +111,12 @@ public:
 	/** Set the last time a frame was received. */
 	void SetLastPushTime(double InLastPushTime) { LastPushTime = InLastPushTime; }
 
+	/** Apply this subject's preprocessors to a frame data. */
+	void PreprocessFrame(FLiveLinkFrameDataStruct& InOutFrameData);
+
+	/** Validates if the incoming frame data is compatible with the static data for this subject. */
+	bool ValidateFrameData(const FLiveLinkFrameDataStruct& InFrameData);
+
 private:
 	int32 FindNewFrame_WorldTime(const FLiveLinkWorldTime& FrameTime) const;
 	int32 FindNewFrame_WorldTimeInternal(const FLiveLinkWorldTime& FrameTime) const;
diff --git a/Engine/Plugins/Animation/LiveLinkHub/Source/LiveLinkHub/Private/LiveLinkHub.cpp b/Engine/Plugins/Animation/LiveLinkHub/Source/LiveLinkHub/Private/LiveLinkHub.cpp
--- a/Engine/Plugins/Animation/LiveLinkHub/Source/LiveLinkHub/Private/LiveLinkHub.cpp
+++ b/Engine/Plugins/Animation/LiveLinkHub/Source/LiveLinkHub/Private/LiveLinkHub.cpp
@@ -156,14 +156,13 @@ void FLiveLinkHub::OnStaticDataReceived_AnyThread(const FLiveLinkSubjectKey& InS
 	LiveLinkProvider->UpdateSubjectStaticData(OverridenName, InRole, MoveTemp(StaticDataCopy));
 }
 
-void FLiveLinkHub::OnFrameDataReceived_AnyThread(const FLiveLinkSubjectKey& InSubjectKey, const FLiveLinkFrameDataStruct& InFrameDataStruct) const
+void FLiveLinkHub::OnFrameDataReceived_AnyThread(const FLiveLinkSubjectKey& InSubjectKey, FLiveLinkFrameDataStruct& InFrameDataStruct) const
 {
 	if (RecordingController->IsRecording())
 	{
 		RecordingController->RecordFrameData(InSubjectKey, InFrameDataStruct);
 	}
 
-	
 	FLiveLinkFrameDataStruct FrameDataCopy;
 	FrameDataCopy.InitializeWith(InFrameDataStruct);
 
diff --git a/Engine/Plugins/Animation/LiveLinkHub/Source/LiveLinkHub/Private/LiveLinkHub.h b/Engine/Plugins/Animation/LiveLinkHub/Source/LiveLinkHub/Private/LiveLinkHub.h
--- a/Engine/Plugins/Animation/LiveLinkHub/Source/LiveLinkHub/Private/LiveLinkHub.h
+++ b/Engine/Plugins/Animation/LiveLinkHub/Source/LiveLinkHub/Private/LiveLinkHub.h
@@ -73,7 +73,7 @@ public:
 private:
 	//~ LiveLink Client delegates
 	void OnStaticDataReceived_AnyThread(const FLiveLinkSubjectKey& InSubjectKey, TSubclassOf<ULiveLinkRole> InRole, const FLiveLinkStaticDataStruct& InStaticDataStruct) const;
-	void OnFrameDataReceived_AnyThread(const FLiveLinkSubjectKey& InSubjectKey, const FLiveLinkFrameDataStruct& InFrameDataStruct) const;
+	void OnFrameDataReceived_AnyThread(const FLiveLinkSubjectKey& InSubjectKey, FLiveLinkFrameDataStruct& InFrameDataStruct) const;
 	void OnSubjectMarkedPendingKill_AnyThread(const FLiveLinkSubjectKey& InSubjectKey) const;
 	void OnSubjectAdded(FLiveLinkSubjectKey InSubjectKey) const;
 	//~ LiveLink Client delegates
diff --git a/Engine/Plugins/Animation/LiveLinkHub/Source/LiveLinkHub/Private/LiveLinkHubClient.cpp b/Engine/Plugins/Animation/LiveLinkHub/Source/LiveLinkHub/Private/LiveLinkHubClient.cpp
--- a/Engine/Plugins/Animation/LiveLinkHub/Source/LiveLinkHub/Private/LiveLinkHubClient.cpp
+++ b/Engine/Plugins/Animation/LiveLinkHub/Source/LiveLinkHub/Private/LiveLinkHubClient.cpp
@@ -37,6 +37,21 @@ FLiveLinkHubClient::~FLiveLinkHubClient()
 {
 }
 
+void FLiveLinkHubClient::CacheSubjectSettings(const FLiveLinkSubjectKey& SubjectKey, ULiveLinkSubjectSettings* Settings) const
+{
+	FScopeLock Lock(&CollectionAccessCriticalSection);
+
+	if (FLiveLinkCollectionSubjectItem* SubjectItem = Collection->FindSubject(SubjectKey))
+	{
+		ULiveLinkSourceSettings* SourceSettings = GetSourceSettings(SubjectKey.Source);
+		ULiveLinkSubjectSettings* SubjectSettings = Settings;
+
+		SubjectItem->GetLiveSubject()->CacheSettings(SourceSettings, SubjectSettings);
+
+		BroadcastStaticDataUpdate(SubjectItem->GetLiveSubject(), SubjectItem->GetSubject()->GetRole(), SubjectItem->GetLiveSubject()->GetStaticData());
+	}
+}
+
 const FLiveLinkStaticDataStruct* FLiveLinkHubClient::GetSubjectStaticData(const FLiveLinkSubjectKey& InSubjectKey)
 {
 	FScopeLock Lock(&CollectionAccessCriticalSection);
@@ -379,7 +394,7 @@ void FLiveLinkHubClient::PushSubjectStaticData_AnyThread(const FLiveLinkSubjectK
 		}
 	}
 
-	OnStaticDataReceivedDelegate_AnyThread.Broadcast(SubjectKey, Role, InStaticData);
+	BroadcastStaticDataUpdate(LiveLinkSubject, Role, InStaticData);
 
 	if (LiveLinkSubject)
 	{
@@ -406,15 +421,37 @@ void FLiveLinkHubClient::PushSubjectFrameData_AnyThread(const FLiveLinkSubjectKe
 {
 	SCOPE_CYCLE_COUNTER(STAT_LiveLinkHub_PushFrameData);
 
-	if (const FLiveLinkCollectionSubjectItem* SubjectItem = Collection->FindSubject(SubjectKey))
+	bool bFrameValid = true;
 	{
-		if (FLiveLinkSubject* LiveSubject = SubjectItem->GetLiveSubject())
+		FScopeLock Lock(&CollectionAccessCriticalSection);
+		if (const FLiveLinkCollectionSubjectItem* SubjectItem = Collection->FindSubject(SubjectKey))
 		{
-			LiveSubject->SetLastPushTime(FApp::GetCurrentTime());
+			if (FLiveLinkSubject* LiveSubject = SubjectItem->GetLiveSubject())
+			{
+				LiveSubject->SetLastPushTime(FApp::GetCurrentTime());
+				bFrameValid = LiveSubject->ValidateFrameData(FrameData);
+				if (bFrameValid)
+				{
+					LiveSubject->PreprocessFrame(FrameData);
+
+					// Translator support (Disabled for now since we don't have a way to handle the static data that's given back by the translator)
+					TArray<ULiveLinkFrameTranslator::FWorkerSharedPtr> Translators = LiveSubject->GetFrameTranslators();
+					if (Translators.Num() && Translators[0].IsValid())
+					{
+						FLiveLinkSubjectFrameData OutTranslatedFrame;
+						Translators[0]->Translate(LiveSubject->GetStaticData(), FrameData, OutTranslatedFrame);
+						FrameData = MoveTemp(OutTranslatedFrame.FrameData);
+					}
+				}
+			}
 		}
 	}
 
-	OnFrameDataReceivedDelegate_AnyThread.Broadcast(SubjectKey, FrameData);
+	if (bFrameValid)
+	{
+		OnFrameDataReceivedDelegate_AnyThread.Broadcast(SubjectKey, FrameData);
+	}
+
 	BroadcastFrameDataUpdate(SubjectKey, FrameData);
 }
 
@@ -449,4 +486,32 @@ void FLiveLinkHubClient::RemoveSubject_AnyThread(const FLiveLinkSubjectKey& InSu
 	FLiveLinkClient::RemoveSubject_AnyThread(InSubjectKey);
 }
 
+void FLiveLinkHubClient::BroadcastStaticDataUpdate(FLiveLinkSubject* InLiveSubject, TSubclassOf<ULiveLinkRole> InRole, const FLiveLinkStaticDataStruct& InStaticData) const
+{
+	// If we have a translator, we need to broadcast the translated static data to the clients.
+	FLiveLinkStaticDataStruct StaticDataStruct;
+	TArray<ULiveLinkFrameTranslator::FWorkerSharedPtr> Translators = InLiveSubject->GetFrameTranslators();
+	if (Translators.Num() && Translators[0].IsValid())
+	{
+
+		ULiveLinkFrameTranslator::FWorkerSharedPtr Translator = Translators[0];
+		TSubclassOf<ULiveLinkRole> FromRole = Translator->GetFromRole();
+		UScriptStruct* FrameDataStruct = FromRole->GetDefaultObject<ULiveLinkRole>()->GetFrameDataStruct();
+		FStructOnScope FrameDataOnScope(FrameDataStruct);
+		FrameDataStruct->InitializeDefaultValue(FrameDataOnScope.GetStructMemory());
+
+		FLiveLinkFrameDataStruct BaseFrameStruct(FrameDataStruct);
+		FrameDataStruct->InitializeDefaultValue(FrameDataOnScope.GetStructMemory());
+
+		FLiveLinkSubjectFrameData OutTranslatedFrame;
+		Translator->Translate(InLiveSubject->GetStaticData(), BaseFrameStruct, OutTranslatedFrame);
+
+		OnStaticDataReceivedDelegate_AnyThread.Broadcast(InLiveSubject->GetSubjectKey(), Translator->GetToRole(), OutTranslatedFrame.StaticData);
+	}
+	else
+	{
+		OnStaticDataReceivedDelegate_AnyThread.Broadcast(InLiveSubject->GetSubjectKey(), InRole, InStaticData);
+	}
+}
+
 #undef LOCTEXT_NAMESPACE
diff --git a/Engine/Plugins/Animation/LiveLinkHub/Source/LiveLinkHub/Private/LiveLinkHubClient.h b/Engine/Plugins/Animation/LiveLinkHub/Source/LiveLinkHub/Private/LiveLinkHubClient.h
--- a/Engine/Plugins/Animation/LiveLinkHub/Source/LiveLinkHub/Private/LiveLinkHubClient.h
+++ b/Engine/Plugins/Animation/LiveLinkHub/Source/LiveLinkHub/Private/LiveLinkHubClient.h
@@ -10,10 +10,11 @@
 
 class FLiveLinkHubPlaybackController;
 class FLiveLinkHubRecordingController;
+class FLiveLinkSubject;
 class ILiveLinkHub;
 struct ILiveLinkProvider;
 
-DECLARE_TS_MULTICAST_DELEGATE_TwoParams(FOnFrameDataReceived_AnyThread, const FLiveLinkSubjectKey&, const FLiveLinkFrameDataStruct&);
+DECLARE_TS_MULTICAST_DELEGATE_TwoParams(FOnFrameDataReceived_AnyThread, const FLiveLinkSubjectKey&, FLiveLinkFrameDataStruct&);
 DECLARE_TS_MULTICAST_DELEGATE_ThreeParams(FOnStaticDataReceived_AnyThread, const FLiveLinkSubjectKey&, TSubclassOf<ULiveLinkRole>, const FLiveLinkStaticDataStruct&);
 DECLARE_TS_MULTICAST_DELEGATE_OneParam(FOnSubjectMarkedPendingKill_AnyThread, const FLiveLinkSubjectKey&);
 
@@ -48,6 +49,9 @@ public:
 		return OnSubjectMarkedPendingKillDelegate_AnyThread;
 	}
 
+	/** Cache subject settings for the subject specified by the subject key. */
+	void CacheSubjectSettings(const FLiveLinkSubjectKey& SubjectKey, ULiveLinkSubjectSettings* Settings) const;
+
 public:
 	//~ Begin ILiveLinkClient interface
 	virtual bool CreateSource(const FLiveLinkSourcePreset& InSourcePreset) override;
@@ -64,6 +68,8 @@ private:
 	bool CreatePlaybackSource(const FLiveLinkSourcePreset& InSourcePreset);
 	/** Create a LiveLinkPlaybackSubject which acts as a dummy subject when doing playback. */
 	bool CreatePlaybackSubject(const FLiveLinkSubjectPreset& InSubjectPreset);
+	/** Broadcast a static data update to this client's listeners. */
+	void BroadcastStaticDataUpdate(FLiveLinkSubject* InLiveSubject, TSubclassOf<ULiveLinkRole> InRole, const FLiveLinkStaticDataStruct& InStaticData) const;
 	/** Lock to stop multiple threads accessing the Subjects from the collection at the same time */
 	mutable FCriticalSection CollectionAccessCriticalSection;
 
diff --git a/Engine/Plugins/Animation/LiveLinkHub/Source/LiveLinkHub/Private/Session/LiveLinkHubSession.h b/Engine/Plugins/Animation/LiveLinkHub/Source/LiveLinkHub/Private/Session/LiveLinkHubSession.h
--- a/Engine/Plugins/Animation/LiveLinkHub/Source/LiveLinkHub/Private/Session/LiveLinkHubSession.h
+++ b/Engine/Plugins/Animation/LiveLinkHub/Source/LiveLinkHub/Private/Session/LiveLinkHubSession.h
@@ -21,6 +21,9 @@
 DECLARE_MULTICAST_DELEGATE_OneParam(FOnClientAddedToSession, FLiveLinkHubClientId);
 DECLARE_MULTICAST_DELEGATE_OneParam(FOnClientRemovedFromSession, FLiveLinkHubClientId);
 
+class ULiveLinkFramePreProcessor;
+class ULiveLinkFrameTranslator;
+
 /**
  * Holds the state of the hub for an active session, can be swapped out with a different session using the session manager.
  */
@@ -44,6 +47,12 @@ public:
 	/** Returns whether a client is in this session. */
 	virtual bool IsClientInSession(const FLiveLinkHubClientId& Client) = 0;
 
+	/** Set the preprocessors for a subject. */
+	virtual void SetPreProcessors(const FLiveLinkSubjectKey& SubjectKey, TConstArrayView<ULiveLinkFramePreProcessor*> PreProcessors) = 0;
+
+	/** Set the translators for a subject. */
+	virtual void SetTranslator(const FLiveLinkSubjectKey& SubjectKey, ULiveLinkFrameTranslator* Translator) = 0;
+
 	/** Get the list of clients in this session (The list of clients that can receive data from the hub) */
 	virtual TArray<FLiveLinkHubClientId> GetSessionClients() const = 0;
 };
@@ -139,6 +148,24 @@ public:
 		return CachedSessionClients.Contains(Client);
 	}
 
+	virtual void SetPreProcessors(const FLiveLinkSubjectKey& SubjectKey, TConstArrayView<ULiveLinkFramePreProcessor*> PreProcessors) override
+	{
+		FWriteScopeLock Locker(SessionDataLock);
+		if (FLiveLinkHubSubjectProxy* SubjectProxy = SessionData.SubjectsConfig.SubjectProxies.Find(SubjectKey))
+		{
+			SubjectProxy->SetPreProcessors(PreProcessors);
+		}
+	}
+
+	virtual void SetTranslator(const FLiveLinkSubjectKey& SubjectKey, ULiveLinkFrameTranslator* Translator) override
+	{
+		FWriteScopeLock Locker(SessionDataLock);
+		if (FLiveLinkHubSubjectProxy* SubjectProxy = SessionData.SubjectsConfig.SubjectProxies.Find(SubjectKey))
+		{
+			SubjectProxy->SetTranslator(Translator);
+		}
+	}
+
 	void AddRestoredClient(FLiveLinkHubUEClientInfo& InOutRestoredClientInfo)
 	{
 		if (const TSharedPtr<FLiveLinkHubProvider> LiveLinkProvider = FModuleManager::Get().GetModuleChecked<FLiveLinkHubModule>("LiveLinkHub").GetLiveLinkProvider())
@@ -195,7 +222,7 @@ private:
 			ILiveLinkClient& LiveLinkClient = IModularFeatures::Get().GetModularFeature<ILiveLinkClient>(ILiveLinkClient::ModularFeatureName);
 
 			FLiveLinkHubSubjectProxy SubjectSettings;
-			SubjectSettings.Initialize(SubjectKey, LiveLinkClient.GetSourceType(SubjectKey.Source).ToString());
+			SubjectSettings.Initialize(SubjectKey, LiveLinkClient.GetSourceType(SubjectKey.Source).ToString(), Cast<ULiveLinkSubjectSettings>(LiveLinkClient.GetSubjectSettings(SubjectKey)));
 
 			{
 				FWriteScopeLock Locker(SessionDataLock);
diff --git a/Engine/Plugins/Animation/LiveLinkHub/Source/LiveLinkHub/Private/Subjects/LiveLinkHubSubjectController.h b/Engine/Plugins/Animation/LiveLinkHub/Source/LiveLinkHub/Private/Subjects/LiveLinkHubSubjectController.h
--- a/Engine/Plugins/Animation/LiveLinkHub/Source/LiveLinkHub/Private/Subjects/LiveLinkHubSubjectController.h
+++ b/Engine/Plugins/Animation/LiveLinkHub/Source/LiveLinkHub/Private/Subjects/LiveLinkHubSubjectController.h
@@ -2,7 +2,6 @@
 
 #pragma once
 
-#include "IStructureDetailsView.h"
 #include "LiveLinkHub.h"
 #include "LiveLinkHubModule.h"
 #include "LiveLinkTypes.h"
@@ -38,17 +37,18 @@ public:
 	TSharedRef<SWidget> MakeSubjectView()
 	{
 		return SAssignNew(SubjectsView, SLiveLinkHubSubjectView, SubjectModel.ToSharedRef())
-			.OnRenameSubject_Raw(this, &FLiveLinkHubSubjectController::OnSubjectRenamed);
+			.OnRenameSubject_Raw(this, &FLiveLinkHubSubjectController::OnSubjectRenamed)
+			.OnProcessorModified_Raw(this, &FLiveLinkHubSubjectController::OnSubjectProcessorModified);
 	}
 
 	/** Set the displayed subject in the subject view. */
-	void SetSubject(const FLiveLinkSubjectKey& Subject)
+	void SetSubject(const FLiveLinkSubjectKey& Subject) const
 	{
 		SubjectsView->SetSubject(Subject);
 	}
 
 	/** Handle modifying the session config for the specified subject. */
-	void OnSubjectRenamed(const FLiveLinkSubjectKey& SubjectKey, FName NewName)
+	void OnSubjectRenamed(const FLiveLinkSubjectKey& SubjectKey, FName NewName) const
 	{
 		if (TSharedPtr<ILiveLinkHubSessionManager> SessionManager = FModuleManager::Get().GetModuleChecked<FLiveLinkHubModule>("LiveLinkHub").GetLiveLinkHub()->GetSessionManager())
 		{
@@ -56,12 +56,59 @@ public:
 		}
 	}
 
+	void OnSubjectProcessorModified(const FLiveLinkSubjectKey& SubjectKey, const TArray<TSubclassOf<ULiveLinkFramePreProcessor>>& UpdatedPreprocessors, const TSubclassOf<ULiveLinkFrameTranslator> UpdatedTranslator) const
+	{
+		FLiveLinkHubClient& LiveLinkClient = static_cast<FLiveLinkHubClient&>(IModularFeatures::Get().GetModularFeature<FLiveLinkClient>(ILiveLinkClient::ModularFeatureName));
+		ULiveLinkSubjectSettings* Settings = Cast<ULiveLinkSubjectSettings>(LiveLinkClient.GetSubjectSettings(SubjectKey));
+
+		Settings->PreProcessors.Reset(UpdatedPreprocessors.Num());
+		for (TSubclassOf<ULiveLinkFramePreProcessor> PreProcessor : UpdatedPreprocessors)
+		{
+			Settings->PreProcessors.Add(PreProcessor.GetDefaultObject());
+		}
+
+		Settings->Translators.Reset();
+
+		if (UpdatedTranslator)
+		{
+			Settings->Translators.Add(UpdatedTranslator.GetDefaultObject());
+		}
+
+		// todo: If we are removing the translator, we will need to do additional handling to restore the original static data that was overriden
+
+		if (Settings->ValidateProcessors())
+		{
+			// Apply to the underlying data from the session.
+			if (const TSharedPtr<ILiveLinkHubSessionManager> SessionManager = FModuleManager::Get().GetModuleChecked<FLiveLinkHubModule>("LiveLinkHub").GetLiveLinkHub()->GetSessionManager())
+			{
+				SessionManager->GetCurrentSession()->SetPreProcessors(SubjectKey, Settings->PreProcessors);
+
+				if (Settings->Translators.Num())
+				{
+					SessionManager->GetCurrentSession()->SetTranslator(SubjectKey, Settings->Translators[0]);
+				}
+				else
+				{
+					SessionManager->GetCurrentSession()->SetTranslator(SubjectKey, nullptr);
+				}
+			}
+
+			// Apply to the livelink client
+			LiveLinkClient.CacheSubjectSettings(SubjectKey, Settings);
+		}
+		else
+		{
+			// Refresh the view since the underlying data might have been rolled back.
+			SubjectsView->SetSubject(SubjectKey);
+		}
+	}
+
 	/** Handle updating the subject details when the session has been swapped out for a different one. */
-	void OnActiveSessionChanged(const TSharedRef<ILiveLinkHubSession>& ActiveSession)
+	void OnActiveSessionChanged(const TSharedRef<ILiveLinkHubSession>& ActiveSession) const
 	{
 		if (SubjectsView)
 		{
-			SubjectsView->ClearSubjectDetails();
+			SubjectsView->RefreshSubjectDetails(ActiveSession);
 		}
 	}
 
diff --git a/Engine/Plugins/Animation/LiveLinkHub/Source/LiveLinkHub/Private/Subjects/LiveLinkHubSubjectSessionConfig.cpp b/Engine/Plugins/Animation/LiveLinkHub/Source/LiveLinkHub/Private/Subjects/LiveLinkHubSubjectSessionConfig.cpp
--- a/Engine/Plugins/Animation/LiveLinkHub/Source/LiveLinkHub/Private/Subjects/LiveLinkHubSubjectSessionConfig.cpp
+++ b/Engine/Plugins/Animation/LiveLinkHub/Source/LiveLinkHub/Private/Subjects/LiveLinkHubSubjectSessionConfig.cpp
@@ -7,6 +7,7 @@
 #include "Features/IModularFeatures.h"
 #include "LiveLinkClient.h"
 #include "LiveLinkHubModule.h"
+#include "LiveLinkSubjectSettings.h"
 #include "Modules/ModuleManager.h"
 #include "UObject/StrongObjectPtr.h"
 
@@ -21,7 +22,9 @@ void FLiveLinkHubSubjectSessionConfig::Initialize()
 	for (const FLiveLinkSubjectKey& SubjectKey : LiveLinkClient.GetSubjects(bIncludeDisabledSubject, bIncludeVirtualSubject))
 	{
 		FLiveLinkHubSubjectProxy SubjectSettings;
-		SubjectSettings.Initialize(SubjectKey, LiveLinkClient.GetSourceType(SubjectKey.Source).ToString());
+
+		ULiveLinkSubjectSettings* Settings = Cast<ULiveLinkSubjectSettings>(LiveLinkClient.GetSubjectSettings(SubjectKey));
+		SubjectSettings.Initialize(SubjectKey, LiveLinkClient.GetSourceType(SubjectKey.Source).ToString(), Settings);
 
 		SubjectProxies.Add(SubjectKey, MoveTemp(SubjectSettings));
 	}
@@ -47,12 +50,24 @@ void FLiveLinkHubSubjectSessionConfig::RenameSubject(const FLiveLinkSubjectKey&
 	}
 }
 
-void FLiveLinkHubSubjectProxy::Initialize(const FLiveLinkSubjectKey& InSubjectKey, FString InSource)
+void FLiveLinkHubSubjectProxy::Initialize(const FLiveLinkSubjectKey& InSubjectKey, FString InSource, ULiveLinkSubjectSettings* InSubjectSettings)
 {
 	SubjectName = InSubjectKey.SubjectName.Name.ToString();
 	SubjectKey = InSubjectKey;
 	OutboundName = InSubjectKey.SubjectName.Name.ToString();
 	Source = MoveTemp(InSource);
+
+	PreProcessors.Reset(InSubjectSettings->PreProcessors.Num());
+	for (const ULiveLinkFramePreProcessor* PreProcessor : InSubjectSettings->PreProcessors)
+	{
+		PreProcessors.Add(PreProcessor->GetClass());
+	}
+
+	Translator = nullptr;
+	if (InSubjectSettings->Translators.Num() && InSubjectSettings->Translators[0])
+	{
+		Translator = InSubjectSettings->Translators[0]->GetClass();
+	}
 }
 
 FName FLiveLinkHubSubjectProxy::GetOutboundName() const
diff --git a/Engine/Plugins/Animation/LiveLinkHub/Source/LiveLinkHub/Private/Subjects/LiveLinkHubSubjectSessionConfig.h b/Engine/Plugins/Animation/LiveLinkHub/Source/LiveLinkHub/Private/Subjects/LiveLinkHubSubjectSessionConfig.h
--- a/Engine/Plugins/Animation/LiveLinkHub/Source/LiveLinkHub/Private/Subjects/LiveLinkHubSubjectSessionConfig.h
+++ b/Engine/Plugins/Animation/LiveLinkHub/Source/LiveLinkHub/Private/Subjects/LiveLinkHubSubjectSessionConfig.h
@@ -2,6 +2,10 @@
 
 #pragma once
 
+#include "Algo/Transform.h"
+#include "LiveLinkFramePreProcessor.h"
+#include "LiveLinkFrameTranslator.h"
+#include "LiveLinkSubjectSettings.h"
 #include "Delegates/Delegate.h"
 #include "LiveLinkTypes.h"
 #include "UObject/Object.h"
@@ -17,7 +21,7 @@ public:
 	GENERATED_BODY()
 
 	/** Initialize from a subject key. */
-	void Initialize(const FLiveLinkSubjectKey& InSubjectKey, FString InSource);
+	void Initialize(const FLiveLinkSubjectKey& InSubjectKey, FString InSource, ULiveLinkSubjectSettings* InSubjectSettings);
 
 	/** Get the outbound name for this subject, allows  */
 	FName GetOutboundName() const;
@@ -25,10 +29,24 @@ public:
 	/** Change the outbound name for this subject proxy. */
 	void SetOutboundName(FName NewName);
 
-	/** Get the outbound name property name. */
-	static FName GetOutboundNamePropertyName()
+	/** Set the preprocessors */
+	void SetPreProcessors(TConstArrayView<ULiveLinkFramePreProcessor*> FramePreprocessors)
 	{
-		return GET_MEMBER_NAME_CHECKED(FLiveLinkHubSubjectProxy, OutboundName);
+		PreProcessors.Reset(FramePreprocessors.Num());
+		Algo::TransformIf(FramePreprocessors, PreProcessors, [](const ULiveLinkFramePreProcessor* PreProcessor) { return PreProcessor; }, [](const ULiveLinkFramePreProcessor* PreProcessor){ return PreProcessor->GetClass(); });
+	}
+
+	/** Set the frame translator */
+	void SetTranslator(ULiveLinkFrameTranslator* InFrameTranslator)
+	{
+		if (InFrameTranslator)
+		{
+			Translator = InFrameTranslator->GetClass();
+		}
+		else
+		{
+			Translator = nullptr;
+		}
 	}
 
 private:
@@ -48,6 +66,14 @@ private:
 	UPROPERTY(VisibleAnywhere, Category = "Subject Details")
 	FString Source;
 
+	/** Preprocessors to apply to the livelink data coming through the hub. */
+	UPROPERTY(config, EditAnywhere, Category="Subject Details")
+	TArray<TSubclassOf<ULiveLinkFramePreProcessor>> PreProcessors;
+
+	/** Translators used to convert livelink data to a different role. */
+	UPROPERTY(config, EditAnywhere, Category="Subject Details")
+	TSubclassOf<ULiveLinkFrameTranslator> Translator = nullptr;
+
 	/** SubjectKey for this subject, */
 	UPROPERTY()
 	FLiveLinkSubjectKey SubjectKey;
@@ -60,6 +86,8 @@ private:
 
 	/* Previous outbound name to be used for noticing clients to remove this entry from their subject list. */
 	FName PreviousOutboundName;
+
+	friend class SLiveLinkHubSubjectView;
 };
 
 /** Config pertaining to livelink hub subjects for a given session. */
diff --git a/Engine/Plugins/Animation/LiveLinkHub/Source/LiveLinkHub/Private/Subjects/SLiveLinkHubSubjectView.h b/Engine/Plugins/Animation/LiveLinkHub/Source/LiveLinkHub/Private/Subjects/SLiveLinkHubSubjectView.h
--- a/Engine/Plugins/Animation/LiveLinkHub/Source/LiveLinkHub/Private/Subjects/SLiveLinkHubSubjectView.h
+++ b/Engine/Plugins/Animation/LiveLinkHub/Source/LiveLinkHub/Private/Subjects/SLiveLinkHubSubjectView.h
@@ -4,17 +4,23 @@
 
 #include "Async/Async.h"
 #include "IStructureDetailsView.h"
-#include "Widgets/SCompoundWidget.h"
+#include "IPropertyRowGenerator.h"
 
 #include "Delegates/DelegateCombinations.h"
 #include "DetailsViewArgs.h"
+#include "Features/IModularFeatures.h"
 #include "IDetailsView.h"
+#include "ISinglePropertyView.h"
 #include "LiveLinkHubSubjectModel.h"
+#include "LiveLinkSubjectSettings.h"
 #include "LiveLinkTypes.h"
 #include "PropertyEditorModule.h"
 #include "Modules/ModuleManager.h"
+#include "Widgets/SCompoundWidget.h"
+#include "Widgets/SNullWidget.h"
 
 DECLARE_DELEGATE_TwoParams(FOnRenameLiveLinkHubSubject, const FLiveLinkSubjectKey& /*SubjectKey*/, FName /*NewName*/);
+DECLARE_DELEGATE_ThreeParams(FOnSubjectProcessorModified, const FLiveLinkSubjectKey& /*SubjectKey*/,  const TArray<TSubclassOf<ULiveLinkFramePreProcessor>>& /*UpdatedPreProcessors*/, TSubclassOf<ULiveLinkFrameTranslator> /*UpdatedTranslator*/);
 
 /**
  * Provides the UI that displays information about a livelink hub subject.
@@ -24,7 +30,9 @@ class SLiveLinkHubSubjectView : public SCompoundWidget
 public:
 
 	SLATE_BEGIN_ARGS(SLiveLinkHubSubjectView) {}
+	SLATE_ARGUMENT(FLiveLinkSubjectKey, SubjectKey)
 	SLATE_EVENT(FOnRenameLiveLinkHubSubject, OnRenameSubject)
+	SLATE_EVENT(FOnSubjectProcessorModified, OnProcessorModified)
 	SLATE_END_ARGS()
 
 	//~ Begin SWidget interface
@@ -32,6 +40,7 @@ public:
 	{
 		SubjectModel = InSubjectModel;
 		OnRenameSubjectDelegate = InArgs._OnRenameSubject;
+		OnProcessorModifiedDelegate = InArgs._OnProcessorModified;
 
 		FDetailsViewArgs DetailsViewArgs;
 		DetailsViewArgs.bUpdatesFromSelection = false;
@@ -53,15 +62,22 @@ public:
 
 		ChildSlot
 		[
-			SettingsDetailsView->GetWidget().ToSharedRef()
+			SNew(SVerticalBox)
+			+ SVerticalBox::Slot()
+			.AutoHeight()
+			[
+				SettingsDetailsView->GetWidget().ToSharedRef()
+			]
 		];
 	}
 
 	/** Clear the subject details. */
-	void ClearSubjectDetails()
+	void RefreshSubjectDetails(const TSharedRef<ILiveLinkHubSession>& ActiveSession)
 	{
 		SubjectData->Reset();
 		SettingsDetailsView->SetStructureData(nullptr);
+
+		SetSubject(SubjectKey);
 	}
 
 	/** Set the subject to be displayed in the details view. */
@@ -90,11 +106,19 @@ public:
 	/** Handler called when a subject property is modified, used to trigger a rename on the session's subject config. */
 	void OnSubjectPropertyModified(const FPropertyChangedEvent& PropertyChangedEvent)
 	{
-		if (SubjectData && PropertyChangedEvent.GetPropertyName() == FLiveLinkHubSubjectProxy::GetOutboundNamePropertyName())
+		if (SubjectData)
 		{
 			if (FLiveLinkHubSubjectProxy* Proxy = SubjectData->Get())
 			{
-				OnRenameSubjectDelegate.ExecuteIfBound(SubjectKey, Proxy->GetOutboundName());
+				if (PropertyChangedEvent.GetPropertyName() == GET_MEMBER_NAME_CHECKED(FLiveLinkHubSubjectProxy, OutboundName))
+				{
+					OnRenameSubjectDelegate.ExecuteIfBound(SubjectKey, Proxy->GetOutboundName());
+				}
+				else if (PropertyChangedEvent.GetPropertyName() == GET_MEMBER_NAME_CHECKED(FLiveLinkHubSubjectProxy, PreProcessors)
+					|| PropertyChangedEvent.GetPropertyName() == GET_MEMBER_NAME_CHECKED(FLiveLinkHubSubjectProxy, Translator))
+				{
+					OnProcessorModifiedDelegate.ExecuteIfBound(SubjectKey, Proxy->PreProcessors, Proxy->Translator);
+				}
 			}
 		}
 	}
@@ -110,4 +134,6 @@ private:
 	TSharedPtr<TStructOnScope<FLiveLinkHubSubjectProxy>> SubjectData;
 	/** Delegate called when the outbound name is changed by the user. */
 	FOnRenameLiveLinkHubSubject OnRenameSubjectDelegate;
+	/** Delegate called when a translator or preprocessor is modified by the user. */
+	FOnSubjectProcessorModified OnProcessorModifiedDelegate;
 };
diff --git a/Engine/Source/Runtime/LiveLinkInterface/Private/LiveLinkSubjectSettings.cpp b/Engine/Source/Runtime/LiveLinkInterface/Private/LiveLinkSubjectSettings.cpp
--- a/Engine/Source/Runtime/LiveLinkInterface/Private/LiveLinkSubjectSettings.cpp
+++ b/Engine/Source/Runtime/LiveLinkInterface/Private/LiveLinkSubjectSettings.cpp
@@ -11,57 +11,70 @@
 
 #if WITH_EDITOR
 DEFINE_LOG_CATEGORY_STATIC(LogLiveLinkSubjectSettings, Warning, Warning);
-void ULiveLinkSubjectSettings::PostEditChangeChainProperty(struct FPropertyChangedChainEvent& PropertyChangedEvent)
+
+bool ULiveLinkSubjectSettings::ValidateProcessors()
 {
-	if (PropertyChangedEvent.Property->GetFName() == GET_MEMBER_NAME_CHECKED(ULiveLinkSubjectSettings, PreProcessors)
-	 || PropertyChangedEvent.Property->GetFName() == GET_MEMBER_NAME_CHECKED(ULiveLinkSubjectSettings, InterpolationProcessor)
-	|| PropertyChangedEvent.Property->GetFName() == GET_MEMBER_NAME_CHECKED(ULiveLinkSubjectSettings, Translators))
+	UClass* RoleClass = Role.Get();
+	if (RoleClass == nullptr)
 	{
-		UClass* RoleClass = Role.Get();
-		if (RoleClass == nullptr)
-		{
-			PreProcessors.Reset();
-			InterpolationProcessor = nullptr;
-			Translators.Reset();
-		}
-		else
+		PreProcessors.Reset();
+		InterpolationProcessor = nullptr;
+		Translators.Reset();
+		return false;
+	}
+	else
+	{
+		bool bValidProcessors = true;
+		for (int32 Index = 0; Index < PreProcessors.Num(); ++Index)
 		{
-			for (int32 Index = 0; Index < PreProcessors.Num(); ++Index)
+			if (ULiveLinkFramePreProcessor* PreProcessor = PreProcessors[Index])
 			{
-				if (ULiveLinkFramePreProcessor* PreProcessor = PreProcessors[Index])
+				check(PreProcessor->GetRole() != nullptr);
+				if (!RoleClass->IsChildOf(PreProcessor->GetRole()))
 				{
-					check(PreProcessor->GetRole() != nullptr);
-					if (!RoleClass->IsChildOf(PreProcessor->GetRole()))
-					{
-						UE_LOG(LogLiveLinkSubjectSettings, Warning, TEXT("Role '%s' is not supported by pre processors '%s'"), *RoleClass->GetName(), *PreProcessor->GetName());
-						PreProcessors[Index] = nullptr;
-					}
+					UE_LOG(LogLiveLinkSubjectSettings, Warning, TEXT("Role '%s' is not supported by pre processors '%s'"), *RoleClass->GetName(), *PreProcessor->GetName());
+					PreProcessors[Index] = nullptr;
+					bValidProcessors = false;
 				}
 			}
+		}
 
-			if (InterpolationProcessor)
+		if (InterpolationProcessor)
+		{
+			check(InterpolationProcessor->GetRole() != nullptr);
+			if (!RoleClass->IsChildOf(InterpolationProcessor->GetRole()))
 			{
-				check(InterpolationProcessor->GetRole() != nullptr);
-				if (!RoleClass->IsChildOf(InterpolationProcessor->GetRole()))
-				{
-					UE_LOG(LogLiveLinkSubjectSettings, Warning, TEXT("Role '%s' is not supported by interpolation '%s'"), *RoleClass->GetName(), *InterpolationProcessor->GetName());
-					InterpolationProcessor = nullptr;
-				}
+				UE_LOG(LogLiveLinkSubjectSettings, Warning, TEXT("Role '%s' is not supported by interpolation '%s'"), *RoleClass->GetName(), *InterpolationProcessor->GetName());
+				InterpolationProcessor = nullptr;
+				bValidProcessors = false;
 			}
+		}
 
-			for (int32 Index = 0; Index < Translators.Num(); ++Index)
+		for (int32 Index = 0; Index < Translators.Num(); ++Index)
+		{
+			if (ULiveLinkFrameTranslator* Translator = Translators[Index])
 			{
-				if (ULiveLinkFrameTranslator* Translator = Translators[Index])
+				check(Translator->GetFromRole() != nullptr);
+				if (!RoleClass->IsChildOf(Translator->GetFromRole()))
 				{
-					check(Translator->GetFromRole() != nullptr);
-					if (!RoleClass->IsChildOf(Translator->GetFromRole()))
-					{
-						UE_LOG(LogLiveLinkSubjectSettings, Warning, TEXT("Role '%s' is not supported by translator '%s'"), *RoleClass->GetName(), *Translator->GetName());
-						Translators[Index] = nullptr;
-					}
+					UE_LOG(LogLiveLinkSubjectSettings, Warning, TEXT("Role '%s' is not supported by translator '%s'"), *RoleClass->GetName(), *Translator->GetName());
+					Translators[Index] = nullptr;
+					bValidProcessors = false;
 				}
 			}
 		}
+
+		return bValidProcessors;
+	}
+}
+
+void ULiveLinkSubjectSettings::PostEditChangeChainProperty(struct FPropertyChangedChainEvent& PropertyChangedEvent)
+{
+	if (PropertyChangedEvent.Property->GetFName() == GET_MEMBER_NAME_CHECKED(ULiveLinkSubjectSettings, PreProcessors)
+	 || PropertyChangedEvent.Property->GetFName() == GET_MEMBER_NAME_CHECKED(ULiveLinkSubjectSettings, InterpolationProcessor)
+	|| PropertyChangedEvent.Property->GetFName() == GET_MEMBER_NAME_CHECKED(ULiveLinkSubjectSettings, Translators))
+	{
+		ValidateProcessors();
 	}
 
 	Super::PostEditChangeChainProperty(PropertyChangedEvent);
diff --git a/Engine/Source/Runtime/LiveLinkInterface/Public/LiveLinkSubjectSettings.h b/Engine/Source/Runtime/LiveLinkInterface/Public/LiveLinkSubjectSettings.h
--- a/Engine/Source/Runtime/LiveLinkInterface/Public/LiveLinkSubjectSettings.h
+++ b/Engine/Source/Runtime/LiveLinkInterface/Public/LiveLinkSubjectSettings.h
@@ -50,6 +50,11 @@ public:
 	UPROPERTY(EditAnywhere, Category = "LiveLink")
     bool bRebroadcastSubject;
 
+	/** Validate PreProcessors, Translators and Interpolation processors. Usually called after a property change event.
+	 * Will revert a given change if it does not match the current subject role.
+	 */
+	bool LIVELINKINTERFACE_API ValidateProcessors();
+
 public:
 	//~ Begin UObject interface
 #if WITH_EDITOR
