LiveLinkHub: Fix timecode provider settings not being saved to the json config.

diff --git a/Engine/Plugins/Animation/LiveLinkHub/Source/LiveLinkHub/Private/Clients/LiveLinkHubProvider.h b/Engine/Plugins/Animation/LiveLinkHub/Source/LiveLinkHub/Private/Clients/LiveLinkHubProvider.h
--- a/Engine/Plugins/Animation/LiveLinkHub/Source/LiveLinkHub/Private/Clients/LiveLinkHubProvider.h
+++ b/Engine/Plugins/Animation/LiveLinkHub/Source/LiveLinkHub/Private/Clients/LiveLinkHubProvider.h
@@ -57,6 +57,9 @@ public:
 	/** Timecode settings that should be shared to connected editors. */
 	void SetTimecodeSettings(FLiveLinkHubTimecodeSettings InSettings);
 
+	/** Retrieve current timecode settings. */
+	const FLiveLinkHubTimecodeSettings& GetTimecodeSettings() const { return TimecodeSettings; }
+
 private:
 	/** Handle a connection message resulting from a livelink hub message bus source connecting to this provider. */
 	void HandleHubConnectMessage(const FLiveLinkHubConnectMessage& Message, const TSharedRef<IMessageContext, ESPMode::ThreadSafe>& Context);
diff --git a/Engine/Plugins/Animation/LiveLinkHub/Source/LiveLinkHub/Private/Session/LiveLinkHubSessionData.h b/Engine/Plugins/Animation/LiveLinkHub/Source/LiveLinkHub/Private/Session/LiveLinkHubSessionData.h
--- a/Engine/Plugins/Animation/LiveLinkHub/Source/LiveLinkHub/Private/Session/LiveLinkHubSessionData.h
+++ b/Engine/Plugins/Animation/LiveLinkHub/Source/LiveLinkHub/Private/Session/LiveLinkHubSessionData.h
@@ -46,4 +46,8 @@ struct FLiveLinkHubPersistedSessionData : public FLiveLinkHubSessionData
 	/** Live link hub client info. */
 	UPROPERTY()
 	TArray<FLiveLinkHubUEClientInfo> Clients;
+
+	/** Timecode settings for the live link hub. */
+	UPROPERTY()
+	FLiveLinkHubTimecodeSettings TimecodeSettings;
 };
diff --git a/Engine/Plugins/Animation/LiveLinkHub/Source/LiveLinkHub/Private/Session/LiveLinkHubSessionManager.h b/Engine/Plugins/Animation/LiveLinkHub/Source/LiveLinkHub/Private/Session/LiveLinkHubSessionManager.h
--- a/Engine/Plugins/Animation/LiveLinkHub/Source/LiveLinkHub/Private/Session/LiveLinkHubSessionManager.h
+++ b/Engine/Plugins/Animation/LiveLinkHub/Source/LiveLinkHub/Private/Session/LiveLinkHubSessionManager.h
@@ -159,7 +159,13 @@ public:
 			LiveLinkHubSessionData.Clients.Add(ClientKeyVal.Value);
 		}
 
-		LastConfigPath = SavePath;
+		const FLiveLinkHubTimecodeSettings& TimecodeSettings = LiveLinkProvider->GetTimecodeSettings();
+		LiveLinkHubSessionData.TimecodeSettings = TimecodeSettings;
+		
+		if (!SavePath.IsEmpty())
+		{
+			LastConfigPath = SavePath;
+		}
 		FEditorDirectories::Get().SetLastDirectory(ELastDirectory::GENERIC_SAVE, FPaths::GetPath(LastConfigPath));
 
 		UE::LiveLinkHub::FileUtilities::Private::SaveConfig(LiveLinkHubSessionData, LastConfigPath);
@@ -230,6 +236,9 @@ private:
 			}
 		}
 
+		SessionData->TimecodeSettings.AssignTimecodeSettingsAsProviderToEngine();
+		LiveLinkProvider->SetTimecodeSettings(SessionData->TimecodeSettings);
+
 		TSharedPtr<FLiveLinkHubSession> CurrentSessionPtr;
 		{
 			FScopeLock Lock(&CurrentSessionCS);
diff --git a/Engine/Plugins/Animation/LiveLinkHub/Source/LiveLinkHub/Private/UI/Widgets/SLiveLinkTimecode.cpp b/Engine/Plugins/Animation/LiveLinkHub/Source/LiveLinkHub/Private/UI/Widgets/SLiveLinkTimecode.cpp
--- a/Engine/Plugins/Animation/LiveLinkHub/Source/LiveLinkHub/Private/UI/Widgets/SLiveLinkTimecode.cpp
+++ b/Engine/Plugins/Animation/LiveLinkHub/Source/LiveLinkHub/Private/UI/Widgets/SLiveLinkTimecode.cpp
@@ -20,6 +20,7 @@
 #include "LiveLinkHubCommands.h"
 #include "LiveLinkHubMessages.h"
 #include "LiveLinkHubModule.h"
+#include "Session/LiveLinkHubSessionManager.h"
 
 #define LOCTEXT_NAMESPACE "LiveLinkHub"
 
@@ -96,6 +97,25 @@ void SLiveLinkTimecode::SendUpdatedTimecodeToEditor()
 	}
 }
 
+void SLiveLinkTimecode::UpdateTimecodeFromSettings()
+{
+	FLiveLinkHubTimecodeSettings Settings;
+	using namespace UE::LiveLinkTimecode::Private;
+
+	const TSharedPtr<FLiveLinkHub> LiveLinkHub = FModuleManager::Get().GetModuleChecked<FLiveLinkHubModule>("LiveLinkHub").GetLiveLinkHub();
+	if (LiveLinkHub.IsValid())
+	{
+		TSharedPtr<FLiveLinkHubProvider> Provider = LiveLinkHub->GetLiveLinkProvider();
+		if (Provider.IsValid())
+		{
+			Settings = Provider->GetTimecodeSettings();
+		}
+	}
+
+	bIsTimecodeSource = Settings.Source != ELiveLinkHubTimecodeSource::NotDefined;
+	ActiveTimecodeSource = Settings.SubjectName.IsNone() ? System24fps : Settings.SubjectName;
+}
+
 void SLiveLinkTimecode::SetTimecodeSource(const FName SourceId)
 {
 	if (SourceId == ActiveTimecodeSource)
@@ -177,9 +197,21 @@ TSharedRef<SWidget> SLiveLinkTimecode::MakeMenu()
 
 void SLiveLinkTimecode::Construct(const FArguments& InArgs)
 {
-	ActiveTimecodeSource = UE::LiveLinkTimecode::Private::System24fps;
 	WorkingClient = (FLiveLinkClient*)&IModularFeatures::Get().GetModularFeature<ILiveLinkClient>(ILiveLinkClient::ModularFeatureName);
+	
+	const FLiveLinkHubModule& LiveLinkHubModule = FModuleManager::Get().GetModuleChecked<FLiveLinkHubModule>("LiveLinkHub");
+	TSharedPtr<ILiveLinkHubSessionManager> SessionManager = LiveLinkHubModule.GetSessionManager();
+	if (SessionManager.IsValid())
+	{
+		SessionManager->OnActiveSessionChanged().AddLambda([this](const TSharedRef<ILiveLinkHubSession>&)
+		{
+			// Update the UI when a config is loaded.
+			UpdateTimecodeFromSettings();
+		});
+	}
 
+	UpdateTimecodeFromSettings();
+	
 	check(WorkingClient);
 	ChildSlot
 	[
diff --git a/Engine/Plugins/Animation/LiveLinkHub/Source/LiveLinkHub/Private/UI/Widgets/SLiveLinkTimecode.h b/Engine/Plugins/Animation/LiveLinkHub/Source/LiveLinkHub/Private/UI/Widgets/SLiveLinkTimecode.h
--- a/Engine/Plugins/Animation/LiveLinkHub/Source/LiveLinkHub/Private/UI/Widgets/SLiveLinkTimecode.h
+++ b/Engine/Plugins/Animation/LiveLinkHub/Source/LiveLinkHub/Private/UI/Widgets/SLiveLinkTimecode.h
@@ -44,6 +44,9 @@ private:
 	/** Sends the current timecode settings to the connected editor(s). */
 	void SendUpdatedTimecodeToEditor();
 
+	/** Update the local timecode properties from the timecode settings. */
+	void UpdateTimecodeFromSettings();
+	
 	/** Do we have timecode source enabled */
 	bool bIsTimecodeSource = false;
 
