LiveLinkHub: Display loaded config name in UI.

diff --git a/Engine/Plugins/Animation/LiveLinkHub/Source/LiveLinkHub/Private/Session/LiveLinkHubSessionManager.h b/Engine/Plugins/Animation/LiveLinkHub/Source/LiveLinkHub/Private/Session/LiveLinkHubSessionManager.h
--- a/Engine/Plugins/Animation/LiveLinkHub/Source/LiveLinkHub/Private/Session/LiveLinkHubSessionManager.h
+++ b/Engine/Plugins/Animation/LiveLinkHub/Source/LiveLinkHub/Private/Session/LiveLinkHubSessionManager.h
@@ -26,7 +26,7 @@ public:
 	/** Delegate called when a UE client is removed from the current session, returning it to the list of discovered clients. */
 	virtual FOnClientRemovedFromSession& OnClientRemovedFromSession() = 0;
 
-	/** Delegate called wehen the active session changes, which will change the list of sources, subjects and clients. */
+	/** Delegate called when the active session changes, which will change the list of sources, subjects and clients. */
 	virtual FOnActiveSessionChanged& OnActiveSessionChanged() = 0;
 
 	/** Get the current session, which holds information about which sources, subjects and clients that should be enabled in the hub at the moment. */
@@ -46,6 +46,9 @@ public:
 
 	/** Returns whether the current session has as already been saved to disk before. */
 	virtual bool CanSaveCurrentSession() const = 0;
+
+	/** Returns the last used config path. */
+	virtual const FString& GetLastConfigPath() const = 0;
 };
 
 class FLiveLinkHubSessionManager : public ILiveLinkHubSessionManager
@@ -204,6 +207,12 @@ public:
 	{
 		return !LastConfigPath.IsEmpty();
 	}
+
+	virtual const FString& GetLastConfigPath() const override
+	{
+		return LastConfigPath;
+	}
+	
 	//~ End LiveLinkHubSessionManager
 
 private:
diff --git a/Engine/Plugins/Animation/LiveLinkHub/Source/LiveLinkHub/Private/UI/Widgets/SLiveLinkHubStatusBar.cpp b/Engine/Plugins/Animation/LiveLinkHub/Source/LiveLinkHub/Private/UI/Widgets/SLiveLinkHubStatusBar.cpp
--- a/Engine/Plugins/Animation/LiveLinkHub/Source/LiveLinkHub/Private/UI/Widgets/SLiveLinkHubStatusBar.cpp
+++ b/Engine/Plugins/Animation/LiveLinkHub/Source/LiveLinkHub/Private/UI/Widgets/SLiveLinkHubStatusBar.cpp
@@ -2,14 +2,18 @@
 
 #include "SLiveLinkHubStatusBar.h"
 
-#include "OutputLogCreationParams.h"
+#include "LiveLinkHub.h"
+#include "Session/LiveLinkHubSessionManager.h"
 
 #include "Framework/Application/SlateApplication.h"
+#include "Misc/Paths.h"
+#include "OutputLogCreationParams.h"
 #include "OutputLogModule.h"
 #include "SWidgetDrawer.h"
 #include "WidgetDrawerConfig.h"
-#include "Widgets/Layout/SBox.h"
 #include "Widgets/Layout/SBorder.h"
+#include "Widgets/Layout/SBox.h"
+#include "Widgets/Text/STextBlock.h"
 
 #define LOCTEXT_NAMESPACE "LiveLinkHubStatusBar"
 
@@ -102,10 +106,25 @@ void SLiveLinkHubStatusBar::Construct(const FArguments& InArgs, FName StatusBarI
 {
 	ChildSlot
 	[
-		SNew(SBox)
-		.HeightOverride(FAppStyle::Get().GetFloat("StatusBar.Height"))
+		SNew(SHorizontalBox)
+		+ SHorizontalBox::Slot()
+		.AutoWidth()
+		.VAlign(VAlign_Center)
 		[
-			MakeWidgetDrawer(StatusBarId)
+			SNew(SBox)
+			.VAlign(VAlign_Center)
+			.HeightOverride(FAppStyle::Get().GetFloat("StatusBar.Height"))
+			[
+				MakeWidgetDrawer(StatusBarId)
+			]
+		]
+		+ SHorizontalBox::Slot()
+		.HAlign(HAlign_Right)
+		.VAlign(VAlign_Center)
+		.Padding(FMargin(4.f, 0.f))
+		[
+			SNew(STextBlock)
+			.Text(this, &SLiveLinkHubStatusBar::GetLoadedConfigText)
 		]
 	];
 }
@@ -143,4 +162,15 @@ TSharedRef<SWidgetDrawer> SLiveLinkHubStatusBar::MakeWidgetDrawer(FName StatusBa
 	return WidgetDrawer.ToSharedRef();
 }
 
+FText SLiveLinkHubStatusBar::GetLoadedConfigText() const
+{
+	if (TSharedPtr<ILiveLinkHubSessionManager> SessionManager = FModuleManager::Get().GetModuleChecked<FLiveLinkHubModule>("LiveLinkHub").GetLiveLinkHub()->GetSessionManager())
+	{
+		const FString FileName = FPaths::GetBaseFilename(SessionManager->GetLastConfigPath());
+		return FileName.IsEmpty() ? LOCTEXT("UntitledConfig", "Untitled") : FText::FromString(FileName);
+	}
+
+	return FText::GetEmpty();
+}
+
 #undef LOCTEXT_NAMESPACE
diff --git a/Engine/Plugins/Animation/LiveLinkHub/Source/LiveLinkHub/Private/UI/Widgets/SLiveLinkHubStatusBar.h b/Engine/Plugins/Animation/LiveLinkHub/Source/LiveLinkHub/Private/UI/Widgets/SLiveLinkHubStatusBar.h
--- a/Engine/Plugins/Animation/LiveLinkHub/Source/LiveLinkHub/Private/UI/Widgets/SLiveLinkHubStatusBar.h
+++ b/Engine/Plugins/Animation/LiveLinkHub/Source/LiveLinkHub/Private/UI/Widgets/SLiveLinkHubStatusBar.h
@@ -36,6 +36,9 @@ private:
 	 * @param StatusBarId The unique status bar id to assign.
 	 */
 	TSharedRef<SWidgetDrawer> MakeWidgetDrawer(FName StatusBarId);
+
+	/** The loaded configuration text to display. */
+	FText GetLoadedConfigText() const;
 	
 private:
 	/** The widget drawer for this status bar. */
diff --git a/Engine/Plugins/Animation/LiveLinkHub/Source/LiveLinkHub/Private/UI/Widgets/SLiveLinkHubTabViewBase.cpp b/Engine/Plugins/Animation/LiveLinkHub/Source/LiveLinkHub/Private/UI/Widgets/SLiveLinkHubTabViewBase.cpp
--- a/Engine/Plugins/Animation/LiveLinkHub/Source/LiveLinkHub/Private/UI/Widgets/SLiveLinkHubTabViewBase.cpp
+++ b/Engine/Plugins/Animation/LiveLinkHub/Source/LiveLinkHub/Private/UI/Widgets/SLiveLinkHubTabViewBase.cpp
@@ -134,7 +134,7 @@ void SLiveLinkHubTabViewBase::Construct(const FArguments& InArgs)
 			[
 				SNew(SBorder)
             	.BorderImage(FAppStyle::GetBrush("ToolPanel.GroupBorder"))
-            	.Padding(5.0f, 0.0f, 5.0f, 7.0f)
+            	.Padding(5.0f, 0.0f, 5.0f, 0.0f)
             	[
 					SNew(SLiveLinkHubStatusBar, StatusBarId)
             	]
