Move LiveLinkHub ModuleRules out of the UnrealBuildTool.Rules namespace.

diff --git a/Engine/Plugins/Animation/LiveLinkHub/Source/LiveLinkHub/LiveLinkHub.Build.cs b/Engine/Plugins/Animation/LiveLinkHub/Source/LiveLinkHub/LiveLinkHub.Build.cs
--- a/Engine/Plugins/Animation/LiveLinkHub/Source/LiveLinkHub/LiveLinkHub.Build.cs
+++ b/Engine/Plugins/Animation/LiveLinkHub/Source/LiveLinkHub/LiveLinkHub.Build.cs
@@ -1,47 +1,46 @@
 // Copyright Epic Games, Inc. All Rights Reserved.
 
-namespace UnrealBuildTool.Rules
+using UnrealBuildTool;
+
+public class LiveLinkHub : ModuleRules
 {
-	public class LiveLinkHub : ModuleRules
+	public LiveLinkHub(ReadOnlyTargetRules Target) : base(Target)
 	{
-		public LiveLinkHub(ReadOnlyTargetRules Target) : base(Target)
+		PublicDependencyModuleNames.AddRange(
+		new string[]
 		{
-			PublicDependencyModuleNames.AddRange(
-			new string[]
-			{
-				"Core",
-				"CoreUObject",
-			});
+			"Core",
+			"CoreUObject",
+		});
 
-			PrivateDependencyModuleNames.AddRange(
-			new string[]
-			{
-				"ApplicationCore",
-				"AssetRegistry",
-				"AssetTools",
-				"ContentBrowser",
-				"ContentBrowserAssetDataSource",
-				"ContentBrowserData",
-				"EditorStyle",
-				"Engine",
-				"InputCore",
-				"Json",
-				"JsonUtilities",
-				"LiveLink",
-				"LiveLinkEditor",
-				"LiveLinkHubMessaging",
-				"LiveLinkInterface",
-				"LiveLinkMessageBusFramework",
-				"Slate",
-				"SlateCore",
-				"StandaloneRenderer",
-				"StructUtils",
-				"TimeManagement",
-				"ToolMenus",
-				"ToolWidgets",
-				"UnrealEd",
-				"OutputLog",
-			});
-		}
+		PrivateDependencyModuleNames.AddRange(
+		new string[]
+		{
+			"ApplicationCore",
+			"AssetRegistry",
+			"AssetTools",
+			"ContentBrowser",
+			"ContentBrowserAssetDataSource",
+			"ContentBrowserData",
+			"EditorStyle",
+			"Engine",
+			"InputCore",
+			"Json",
+			"JsonUtilities",
+			"LiveLink",
+			"LiveLinkEditor",
+			"LiveLinkHubMessaging",
+			"LiveLinkInterface",
+			"LiveLinkMessageBusFramework",
+			"Slate",
+			"SlateCore",
+			"StandaloneRenderer",
+			"StructUtils",
+			"TimeManagement",
+			"ToolMenus",
+			"ToolWidgets",
+			"UnrealEd",
+			"OutputLog",
+		});
 	}
 }
