LiveLinkHub - Add settings to the hub. These can be accessed by the settings combo button on the top right of the app.

diff --git a/Engine/Plugins/Animation/LiveLinkHub/Source/LiveLinkHub/LiveLinkHub.Build.cs b/Engine/Plugins/Animation/LiveLinkHub/Source/LiveLinkHub/LiveLinkHub.Build.cs
--- a/Engine/Plugins/Animation/LiveLinkHub/Source/LiveLinkHub/LiveLinkHub.Build.cs
+++ b/Engine/Plugins/Animation/LiveLinkHub/Source/LiveLinkHub/LiveLinkHub.Build.cs
@@ -13,34 +13,38 @@ public class LiveLinkHub : ModuleRules
 			"CoreUObject",
 		});
 
-		PrivateDependencyModuleNames.AddRange(
-		new string[]
-		{
-			"ApplicationCore",
-			"AssetRegistry",
-			"AssetTools",
-			"ContentBrowser",
-			"ContentBrowserAssetDataSource",
-			"ContentBrowserData",
-			"EditorStyle",
-			"Engine",
-			"InputCore",
-			"Json",
-			"JsonUtilities",
-			"LiveLink",
-			"LiveLinkEditor",
-			"LiveLinkHubMessaging",
-			"LiveLinkInterface",
-			"LiveLinkMessageBusFramework",
-			"Slate",
-			"SlateCore",
-			"StandaloneRenderer",
-			"StructUtils",
-			"TimeManagement",
-			"ToolMenus",
-			"ToolWidgets",
-			"UnrealEd",
-			"OutputLog",
-		});
+			PrivateDependencyModuleNames.AddRange(
+			new string[]
+			{
+				"ApplicationCore",
+				"AssetRegistry",
+				"AssetTools",
+				"ContentBrowser",
+				"ContentBrowserAssetDataSource",
+				"ContentBrowserData",
+				"EditorStyle",
+				"Engine",
+				"InputCore",
+				"Json",
+				"JsonUtilities",
+				"LiveLink",
+				"LiveLinkEditor",
+				"LiveLinkHubMessaging",
+				"LiveLinkInterface",
+				"LiveLinkMessageBusFramework",
+				"ProjectSettingsViewer",
+				"SharedSettingsWidgets",
+				"Settings",
+				"SettingsEditor",
+				"Slate",
+				"SlateCore",
+				"StandaloneRenderer",
+				"StructUtils",
+				"TimeManagement",
+				"ToolMenus",
+				"ToolWidgets",
+				"UnrealEd",
+				"OutputLog",
+			});
 	}
 }
diff --git a/Engine/Plugins/Animation/LiveLinkHub/Source/LiveLinkHub/Private/LiveLinkHub.cpp b/Engine/Plugins/Animation/LiveLinkHub/Source/LiveLinkHub/Private/LiveLinkHub.cpp
--- a/Engine/Plugins/Animation/LiveLinkHub/Source/LiveLinkHub/Private/LiveLinkHub.cpp
+++ b/Engine/Plugins/Animation/LiveLinkHub/Source/LiveLinkHub/Private/LiveLinkHub.cpp
@@ -12,10 +12,14 @@
 #include "Features/IModularFeatures.h"
 #include "Framework/Application/SlateApplication.h"
 #include "IDesktopPlatform.h"
+#include "ISettingsModule.h"
+#include "LiveLinkEditorSettings.h"
 #include "LiveLinkHubClient.h"
+#include "LiveLinkHubSettings.h"
 #include "LiveLinkProvider.h"
 #include "LiveLinkHubCommands.h"
 #include "LiveLinkProviderImpl.h"
+#include "LiveLinkSettings.h"
 #include "LiveLinkSubject.h"
 #include "LiveLinkSubjectSettings.h"
 #include "Misc/App.h"
@@ -43,6 +47,8 @@ void FLiveLinkHub::Initialize()
 	SessionManager = MakeShared<FLiveLinkHubSessionManager>();
 	LiveLinkProvider = MakeShared<FLiveLinkHubProvider>(SessionManager.ToSharedRef());
 
+	FModuleManager::Get().LoadModule("Settings");
+
 	IAssetRegistry& AssetRegistry = FModuleManager::LoadModuleChecked<FAssetRegistryModule>(AssetRegistryConstants::ModuleName).Get();
 	const FString FilePath = FPaths::Combine(FPlatformProcess::UserSettingsDir(), *FApp::GetEpicProductIdentifier(), TEXT("LiveLinkHub"), TEXT("Content"));
 	AssetRegistry.ScanPathsSynchronous({ FilePath }, /*bForceRescan=*/ true);
@@ -73,6 +79,8 @@ void FLiveLinkHub::Initialize()
 	LiveLinkHubClient->OnSubjectMarkedPendingKill_AnyThread().AddSP(this, &FLiveLinkHub::OnSubjectMarkedPendingKill_AnyThread);
 	LiveLinkHubClient->OnLiveLinkSubjectAdded().AddSP(this, &FLiveLinkHub::OnSubjectAdded);
 
+	RegisterLiveLinkHubSettings();
+
 	PlaybackController->Start();
 
 	GIsRunning = true;
@@ -80,6 +88,8 @@ void FLiveLinkHub::Initialize()
 
 FLiveLinkHub::~FLiveLinkHub()
 {
+	UnregisterLiveLinkHubSettings();
+
 	RecordingController.Reset();
 	PlaybackController.Reset();
 
@@ -257,4 +267,38 @@ FName FLiveLinkHub::GetSubjectNameOverride(const FLiveLinkSubjectKey& InSubjectK
 	return InSubjectKey.SubjectName;
 }
 
+void FLiveLinkHub::RegisterLiveLinkHubSettings()
+{
+	if (ISettingsModule* SettingsModule = FModuleManager::GetModulePtr<ISettingsModule>("Settings"))
+	{
+		SettingsModule->RegisterSettings("Editor", "Plugins", "LiveLink",
+			LOCTEXT("EditorSettingsName", "Live Link"),
+			LOCTEXT("EditorSettingsDescription", "Configure LiveLink."),
+			GetMutableDefault<ULiveLinkEditorSettings>()
+		);
+
+		SettingsModule->RegisterSettings("Project", "Plugins", "LiveLink",
+			LOCTEXT("LiveLinkSettingsName", "Live Link"),
+			LOCTEXT("LiveLinkDescription", "Configure LiveLink."),
+			GetMutableDefault<ULiveLinkSettings>()
+		);
+
+		SettingsModule->RegisterSettings("Project", "Plugins", "LiveLinkHub",
+			LOCTEXT("LiveLinkHubSettingsName", "Live Link Hub"),
+			LOCTEXT("LiveLinkHubDescription", "Configure LiveLink Hub."),
+			GetMutableDefault<ULiveLinkHubSettings>()
+		);
+	}
+}
+
+void FLiveLinkHub::UnregisterLiveLinkHubSettings()
+{
+	if (ISettingsModule* SettingsModule = FModuleManager::GetModulePtr<ISettingsModule>("Settings"))
+	{
+		SettingsModule->UnregisterSettings("Editor", "Plugins", "LiveLink");
+		SettingsModule->UnregisterSettings("Project", "Plugins", "LiveLink");
+		SettingsModule->UnregisterSettings("Project", "Plugins", "LiveLinkHub");
+	}
+}
+
 #undef LOCTEXT_NAMESPACE /*LiveLinkHub*/
diff --git a/Engine/Plugins/Animation/LiveLinkHub/Source/LiveLinkHub/Private/LiveLinkHub.h b/Engine/Plugins/Animation/LiveLinkHub/Source/LiveLinkHub/Private/LiveLinkHub.h
--- a/Engine/Plugins/Animation/LiveLinkHub/Source/LiveLinkHub/Private/LiveLinkHub.h
+++ b/Engine/Plugins/Animation/LiveLinkHub/Source/LiveLinkHub/Private/LiveLinkHub.h
@@ -83,6 +83,10 @@ private:
 
 	/** Clear all client settings. */
 	FName GetSubjectNameOverride(const FLiveLinkSubjectKey& InSubjectKey) const;
+	/** Register settings relevant to the livelink hub. */
+	void RegisterLiveLinkHubSettings();
+	/** Unregister settings relevant to the livelink hub. */
+	void UnregisterLiveLinkHubSettings();
 
 	/** Create a new config. */
 	void NewConfig();
diff --git a/Engine/Plugins/Animation/LiveLinkHub/Source/LiveLinkHub/Private/LiveLinkHubSettings.h b/Engine/Plugins/Animation/LiveLinkHub/Source/LiveLinkHub/Private/LiveLinkHubSettings.h
--- a/Engine/Plugins/Animation/LiveLinkHub/Source/LiveLinkHub/Private/LiveLinkHubSettings.h
+++ b/Engine/Plugins/Animation/LiveLinkHub/Source/LiveLinkHub/Private/LiveLinkHubSettings.h
@@ -9,13 +9,17 @@
 /**
  * Settings for LiveLinkHub.
  */
-UCLASS(config=Game, defaultconfig)
+UCLASS(config=Engine, defaultconfig)
 class LIVELINKHUB_API ULiveLinkHubSettings : public UObject
 {
 	GENERATED_BODY()
 
 public:
 	/** If enabled, discovered clients will be automatically added to the current session. */
-	UPROPERTY(config, EditAnywhere, Category="LiveLinkHub")
+	UPROPERTY(config, EditAnywhere, Category="LiveLinkHub", DisplayName = "Automatically add discovered clients")
 	bool bAutoAddDiscoveredClients = true;
+
+	/** Which project settings sections to display when opening the settings viewer. */
+	UPROPERTY(config)
+	TArray<FName> ProjectSettingsToDisplay;
 };
diff --git a/Engine/Plugins/Animation/LiveLinkHub/Source/LiveLinkHub/Private/UI/Widgets/SLiveLinkHubSettings.h b/Engine/Plugins/Animation/LiveLinkHub/Source/LiveLinkHub/Private/UI/Widgets/SLiveLinkHubSettings.h
new file mode 100644
--- /dev/null
+++ b/Engine/Plugins/Animation/LiveLinkHub/Source/LiveLinkHub/Private/UI/Widgets/SLiveLinkHubSettings.h
@@ -0,0 +1,183 @@
+// Copyright Epic Games, Inc. All Rights Reserved.
+
+#pragma once
+
+#include "EditorFontGlyphs.h"
+#include "Engine/DeveloperSettings.h"
+#include "Framework/Docking/TabManager.h"
+#include "Framework/MultiBox/MultiBoxBuilder.h"
+#include "Modules/ModuleManager.h"
+#include "ISettingsCategory.h"
+#include "ISettingsContainer.h"
+#include "ISettingsEditorModel.h"
+#include "ISettingsEditorModule.h"
+#include "ISettingsModule.h"
+#include "ISettingsSection.h"
+#include "LiveLinkHubSettings.h"
+#include "Widgets/Docking/SDockTab.h"
+#include "Widgets/Input/SComboBox.h"
+#include "Widgets/SCompoundWidget.h"
+
+#define LOCTEXT_NAMESPACE "LiveLinkHubSettings"
+
+/**
+ * LiveLinkHub Settings widget.
+ */
+class SLiveLinkHubSettings : public SCompoundWidget
+{
+public:
+
+	SLATE_BEGIN_ARGS(SLiveLinkHubSettings) {}
+	SLATE_END_ARGS()
+
+	/** ID for the settings tab. */
+	const FName SettingsTabId = "LiveLinkHubSettings";
+
+	/**
+	 * Create the settings combo button.
+	 */
+	void Construct(const FArguments& InArgs)
+	{
+		UnregisterSettings();
+
+		// Settings editor might not be loaded at this point.
+		ISettingsEditorModule* SettingsEditorModule = static_cast<ISettingsEditorModule*>(FModuleManager::Get().LoadModule("SettingsEditor"));
+		SettingsEditorModule->SetShouldRegisterSettingCallback(FShouldRegisterSettingsDelegate::CreateLambda([](UDeveloperSettings* Settings) { return false; }));
+
+		FGlobalTabmanager::Get()->RegisterTabSpawner(SettingsTabId, FOnSpawnTab::CreateRaw(this, &SLiveLinkHubSettings::SpawnSettingsTab))
+			.SetDisplayName(LOCTEXT("LiveLinkHubSettingsTabLabel", "Settings"))
+			.SetIcon(FSlateIcon(FAppStyle::GetAppStyleSetName(), "LevelEditor.Tabs.Details"));
+
+		ChildSlot
+		[
+			SNew(SComboButton)
+			.ContentPadding(FMargin(2.0f, 4.0f))
+			.ButtonStyle(&FAppStyle::Get().GetWidgetStyle<FButtonStyle>("Button"))
+			.ForegroundColor(FSlateColor::UseStyle())
+			.ButtonContent()
+			[
+				SNew(SHorizontalBox)
+				+ SHorizontalBox::Slot()
+				.VAlign(VAlign_Center)
+				.Padding(4.0f, 0.0f)
+				.AutoWidth()
+				[
+					SNew(STextBlock)
+					.Font(FAppStyle::Get().GetFontStyle("FontAwesome.14"))
+					.Text(FEditorFontGlyphs::Cogs)
+				]
+				+ SHorizontalBox::Slot()
+				.VAlign(VAlign_Center)
+				.Padding(2.0f, 0, 0, 0)
+				.AutoWidth()
+				[
+					SNew(STextBlock)
+					.TextStyle(FAppStyle::Get(), "SmallButtonText")
+					.Text(LOCTEXT("SettingsLabel", "Settings"))
+				]
+			]
+			.OnGetMenuContent(this, &SLiveLinkHubSettings::OnSettingsComboBoxOpen)
+		];
+	}
+
+	~SLiveLinkHubSettings()
+	{
+		FGlobalTabmanager::Get()->UnregisterTabSpawner(SettingsTabId);
+	}
+
+	/** Unregister all settings except the ones that are relevant to the hub. */
+	void UnregisterSettings()
+	{
+		ISettingsModule& SettingsModule = FModuleManager::GetModuleChecked<ISettingsModule>("Settings");
+
+		TArray<FName> ContainerNames;
+		SettingsModule.GetContainerNames(ContainerNames);
+
+		const TArray<FName>& SettingsToDisplay = GetDefault<ULiveLinkHubSettings>()->ProjectSettingsToDisplay;
+
+		for (FName ContainerName : ContainerNames)
+		{
+			if (TSharedPtr<ISettingsContainer> Container = SettingsModule.GetContainer(ContainerName))
+			{
+				TArray<TSharedPtr<ISettingsCategory>> Categories;
+				Container->GetCategories(Categories);
+
+				for (TSharedPtr<ISettingsCategory> SettingsCategory : Categories)
+				{
+					TArray<TSharedPtr<ISettingsSection>> Sections;
+					constexpr bool bIgnoreVisibility = true;
+					SettingsCategory->GetSections(Sections, bIgnoreVisibility);
+
+					for (TSharedPtr<ISettingsSection> Section : Sections)
+					{
+						if (!SettingsToDisplay.Contains(Section->GetName()))
+						{
+							SettingsModule.UnregisterSettings(Container->GetName(), SettingsCategory->GetName(), Section->GetName());
+						}
+					}
+				}
+			}
+		}
+	}
+
+	/** Handler called when the combo button is clicked. */
+	TSharedRef<SWidget> OnSettingsComboBoxOpen() const
+	{
+		const FUIAction OpenHubSettingsAction(
+			FExecuteAction::CreateLambda([TabId = SettingsTabId]
+				{
+					if (ISettingsModule* SettingsModule = FModuleManager::GetModulePtr<ISettingsModule>("Settings"))
+					{
+						FGlobalTabmanager::Get()->TryInvokeTab(TabId);
+					}
+				}), FCanExecuteAction(), FIsActionChecked());
+
+		const bool CloseAfterSelection = true;
+		FMenuBuilder MenuBuilder(CloseAfterSelection, NULL);
+
+		MenuBuilder.BeginSection("HubSettings", LOCTEXT("HubSettingsLabel", "Hub Settings"));
+		MenuBuilder.AddMenuEntry(
+			LOCTEXT("SettingsMenuEntry", "Settings"),
+			LOCTEXT("SettingsMenuEntry_ToolTip", "Open the settings menu."),
+			FSlateIcon(),
+			OpenHubSettingsAction,
+			NAME_None,
+			EUserInterfaceActionType::Button);
+
+		MenuBuilder.EndSection();
+
+		return MenuBuilder.MakeWidget();
+	}
+
+	/** Creates the settings viewer tab. */
+	TSharedRef<SDockTab> SpawnSettingsTab(const FSpawnTabArgs& SpawnTabArgs)
+	{
+		TSharedRef<SWidget> SettingsEditor = SNullWidget::NullWidget;
+		ISettingsModule* SettingsModule = FModuleManager::GetModulePtr<ISettingsModule>("Settings");
+
+		if (SettingsModule != nullptr)
+		{
+			ISettingsContainerPtr SettingsContainer = SettingsModule->GetContainer("Project");
+
+			if (SettingsContainer.IsValid())
+			{
+				ISettingsEditorModule& SettingsEditorModule = FModuleManager::GetModuleChecked<ISettingsEditorModule>("SettingsEditor");
+
+				SettingsEditorModel = SettingsEditorModule.CreateModel(SettingsContainer.ToSharedRef());
+				SettingsEditor = SettingsEditorModule.CreateEditor(SettingsEditorModel.ToSharedRef());
+			}
+		}
+
+		return SNew(SDockTab)
+			.TabRole(ETabRole::NomadTab)
+			[
+				SettingsEditor
+			];
+	}
+
+private:
+	/** Holds the view state for the settings. */
+	TSharedPtr<ISettingsEditorModel> SettingsEditorModel;
+};
+
+#undef LOCTEXT_NAMESPACE /*LiveLinkHubSettings*/
diff --git a/Engine/Plugins/Animation/LiveLinkHub/Source/LiveLinkHub/Private/UI/Widgets/SLiveLinkHubTabViewBase.cpp b/Engine/Plugins/Animation/LiveLinkHub/Source/LiveLinkHub/Private/UI/Widgets/SLiveLinkHubTabViewBase.cpp
--- a/Engine/Plugins/Animation/LiveLinkHub/Source/LiveLinkHub/Private/UI/Widgets/SLiveLinkHubTabViewBase.cpp
+++ b/Engine/Plugins/Animation/LiveLinkHub/Source/LiveLinkHub/Private/UI/Widgets/SLiveLinkHubTabViewBase.cpp
@@ -9,6 +9,7 @@
 #include "Misc/App.h"
 #include "Modules/ModuleManager.h"
 #include "Recording/LiveLinkHubRecordingController.h"
+#include "SLiveLinkHubSettings.h"
 #include "SLiveLinkHubStatusBar.h"
 #include "SLiveLinkTimecode.h"
 
@@ -24,6 +25,7 @@
 #define LOCTEXT_NAMESPACE "LiveLinkHub"
 
 #define LIVELINKHUB_SUPPORTS_LAYOUTS 0
+
 TSharedRef<SWidget> GetModeSwitcherContent()
 {
 	FMenuBuilder MenuBuilder(true, NULL);
@@ -109,11 +111,18 @@ void SLiveLinkHubTabViewBase::Construct(const FArguments& InArgs)
 							.Orientation(Orient_Vertical)
 						]
 						+ SHorizontalBox::Slot()
-						.Padding(2.f)
+						.AutoWidth()
 	                    .VAlign(VAlign_Center)
 						[
 							SNew(SLiveLinkTimecode)
 						]
+						+ SHorizontalBox::Slot()
+						.Padding(2.f)
+						.AutoWidth()
+						.VAlign(VAlign_Center)
+						[
+							SNew(SLiveLinkHubSettings)
+						]
 					]
 				]
 			]
diff --git a/Engine/Plugins/Animation/LiveLinkHub/Source/LiveLinkHub/Private/UI/Widgets/SLiveLinkTimecode.cpp b/Engine/Plugins/Animation/LiveLinkHub/Source/LiveLinkHub/Private/UI/Widgets/SLiveLinkTimecode.cpp
--- a/Engine/Plugins/Animation/LiveLinkHub/Source/LiveLinkHub/Private/UI/Widgets/SLiveLinkTimecode.cpp
+++ b/Engine/Plugins/Animation/LiveLinkHub/Source/LiveLinkHub/Private/UI/Widgets/SLiveLinkTimecode.cpp
@@ -216,7 +216,7 @@ void SLiveLinkTimecode::Construct(const FArguments& InArgs)
 	ChildSlot
 	[
 		SNew(SComboButton)
-		.ContentPadding(FMargin(6.0f, 0.0f))
+		.ContentPadding(FMargin(4.0f, 0.0f))
 		.MenuPlacement(MenuPlacement_AboveAnchor)
 		.OnGetMenuContent(this, &SLiveLinkTimecode::MakeMenu)
 		.ComboButtonStyle(&FAppStyle::Get().GetWidgetStyle<FComboButtonStyle>("SimpleComboButton"))
diff --git a/Engine/Programs/LiveLinkHub/Config/DefaultEngine.ini b/Engine/Programs/LiveLinkHub/Config/DefaultEngine.ini
--- a/Engine/Programs/LiveLinkHub/Config/DefaultEngine.ini
+++ b/Engine/Programs/LiveLinkHub/Config/DefaultEngine.ini
@@ -31,6 +31,8 @@ TemplateMacPlist=(FilePath="/Engine/Source/Programs/LiveLinkHub/Resources/Mac/In
 
 [/Script/LiveLinkHub.LiveLinkHubSettings]
 bAutoAddDiscoveredClients=true
++ProjectSettingsToDisplay=LiveLink
++ProjectSettingsToDisplay=LiveLinkHub
 
 [/Script/AppleARKit.AppleARKitSettings]
 bEnableLiveLinkForFaceTracking=true
diff --git a/Engine/Source/Programs/LiveLinkHub/Private/LiveLinkHubCommandlet.h b/Engine/Source/Programs/LiveLinkHub/Private/LiveLinkHubCommandlet.h
--- a/Engine/Source/Programs/LiveLinkHub/Private/LiveLinkHubCommandlet.h
+++ b/Engine/Source/Programs/LiveLinkHub/Private/LiveLinkHubCommandlet.h
@@ -31,10 +31,11 @@ public:
 		{
 			PRIVATE_GIsRunningCommandlet = false;
 			AssetDataSourceModule->StartupModule();
-			PRIVATE_GIsRunningCommandlet = true;	
+			PRIVATE_GIsRunningCommandlet = true;
 		}
 
 		FModuleManager::Get().LoadModule("OutputLog");
+		FModuleManager::Get().LoadModule("ProjectSettingsViewer");
 
 		FModuleManager::Get().LoadModuleChecked<ILiveLinkHubModule>("LiveLinkEditor");
 		FModuleManager::Get().LoadModuleChecked<ILiveLinkHubModule>("LiveLinkHub").StartLiveLinkHub();
