LiveLinkHub: Hide recording selection when no recordings are present.

diff --git a/Engine/Plugins/Animation/LiveLinkHub/Source/LiveLinkHub/Private/Recording/SLiveLinkHubRecordingListView.h b/Engine/Plugins/Animation/LiveLinkHub/Source/LiveLinkHub/Private/Recording/SLiveLinkHubRecordingListView.h
--- a/Engine/Plugins/Animation/LiveLinkHub/Source/LiveLinkHub/Private/Recording/SLiveLinkHubRecordingListView.h
+++ b/Engine/Plugins/Animation/LiveLinkHub/Source/LiveLinkHub/Private/Recording/SLiveLinkHubRecordingListView.h
@@ -2,16 +2,16 @@
 
 #pragma once
 
+#include "AssetRegistry/AssetRegistryModule.h"
+#include "AssetRegistry/IAssetRegistry.h"
 #include "ContentBrowserModule.h"
 #include "Delegates/DelegateCombinations.h"
 #include "Framework/MultiBox/MultiBoxBuilder.h"
 #include "IContentBrowserSingleton.h"
+#include "Implementations/LiveLinkUAssetRecording.h"
 #include "LiveLinkHub.h"
 #include "Recording/LiveLinkRecording.h"
-#include "Recording/LiveLinkHubPlaybackController.h"
-#include "Styling/SlateTypes.h"
 #include "Widgets/Input/SButton.h"
-#include "Widgets/Input/SCheckBox.h"
 #include "Widgets/SCompoundWidget.h"
 #include "Widgets/Text/STextBlock.h"
 
@@ -32,13 +32,36 @@ public:
 		SLATE_EVENT(FCanEject, CanEject)
 	SLATE_END_ARGS()
 
+	SLiveLinkHubRecordingListView()
+	{
+		IAssetRegistry& AssetRegistry = FModuleManager::LoadModuleChecked<FAssetRegistryModule>("AssetRegistry").Get();
+		OnAssetAddedHandle = AssetRegistry.OnAssetAdded().AddRaw(this, &SLiveLinkHubRecordingListView::OnAssetAdded);
+		OnAssetRemovedHandle = AssetRegistry.OnAssetRemoved().AddRaw(this, &SLiveLinkHubRecordingListView::OnAssetRemoved);
+	}
+
+	virtual ~SLiveLinkHubRecordingListView() override
+	{
+		if (FModuleManager::Get().IsModuleLoaded("AssetRegistry"))
+		{
+			IAssetRegistry& AssetRegistry = FModuleManager::GetModuleChecked<FAssetRegistryModule>("AssetRegistry").Get();
+			if (OnAssetAddedHandle.IsValid())
+			{
+				AssetRegistry.OnAssetAdded().Remove(OnAssetAddedHandle);
+			}
+			if (OnAssetRemovedHandle.IsValid())
+			{
+				AssetRegistry.OnAssetRemoved().Remove(OnAssetRemovedHandle);
+			}
+		}
+	}
+	
 	//~ Begin SWidget interface
 	void Construct(const FArguments& InArgs)
 	{
 		OnImportRecordingDelegate = InArgs._OnImportRecording;
 		OnEjectDelegate = InArgs._OnEject;
 		OnCanEjectDelegate = InArgs._CanEject;
-
+		
 		ChildSlot
 		[
 			SNew(SVerticalBox)
@@ -61,6 +84,7 @@ public:
 				.Padding(4.f)
 				[
 					SNew(SButton)
+					.Visibility(this, &SLiveLinkHubRecordingListView::GetRecordingPickerVisibility)
 					.Text(LOCTEXT("EjectButton", "Exit Playback"))
 					.OnClicked(this, &SLiveLinkHubRecordingListView::OnEjectClicked)
 					.IsEnabled(this, &SLiveLinkHubRecordingListView::CanEjectRecording)
@@ -69,7 +93,29 @@ public:
 			+ SVerticalBox::Slot()
 			.FillHeight(1.0f)
 			[
-				CreateRecordingPicker()
+				SNew(SVerticalBox)
+				+ SVerticalBox::Slot()
+				.AutoHeight()
+				[
+					SNew(SBox)
+					.Visibility(this, &SLiveLinkHubRecordingListView::GetRecordingPickerVisibility)
+					[
+						CreateRecordingPicker()
+					]
+				]
+				+SVerticalBox::Slot()
+				.AutoHeight()
+				.VAlign(VAlign_Top)
+				.HAlign(HAlign_Center)
+				[
+					SNew(STextBlock)
+					.Visibility_Lambda([this]()
+					{
+						EVisibility RecordingPickerVisibility = GetRecordingPickerVisibility();
+						return RecordingPickerVisibility == EVisibility::Visible ? EVisibility::Collapsed : EVisibility::Visible;
+					})
+					.Text(GetNoAssetsWarningText())
+				]
 			]
 		];
 	}
@@ -94,13 +140,53 @@ private:
 		return OnCanEjectDelegate.Execute();
 	}
 
+	/** When an asset is added to the asset registry. */
+	void OnAssetAdded(const FAssetData& InAssetData)
+	{
+		if (InAssetData.IsValid()
+			&& (InAssetData.AssetClassPath == ULiveLinkUAssetRecording::StaticClass()->GetClassPathName()
+				|| InAssetData.AssetClassPath == ULiveLinkRecording::StaticClass()->GetClassPathName()))
+		{
+			bAssetsAvailableCached = true;
+		}
+	}
+
+	/** When an asset is removed from the asset registry. */
+	void OnAssetRemoved(const FAssetData& InAssetData)
+	{
+		if (InAssetData.IsValid()
+			&& (InAssetData.AssetClassPath == ULiveLinkUAssetRecording::StaticClass()->GetClassPathName()
+				|| InAssetData.AssetClassPath == ULiveLinkRecording::StaticClass()->GetClassPathName()))
+		{
+			// Let the cache recalculate.
+			bAssetsAvailableCached.Reset();
+		}
+	}
+
+	/** The visibility status of the recording picker. */
+	EVisibility GetRecordingPickerVisibility() const
+	{
+		// Cache the value initially, otherwise it is set on the asset added event.
+		if (!bAssetsAvailableCached.IsSet())
+		{
+			FAssetRegistryModule& AssetRegistryModule = FModuleManager::LoadModuleChecked<FAssetRegistryModule>("AssetRegistry");
+			TArray<FAssetData> AssetDataArray;
+			FARFilter Filter = MakeAssetFilter();
+			AssetRegistryModule.Get().GetAssets(MoveTemp(Filter), AssetDataArray);
+
+			bAssetsAvailableCached = AssetDataArray.Num() > 0;
+		}
+		
+		return bAssetsAvailableCached.GetValue() ? EVisibility::Visible : EVisibility::Hidden;
+	}
+
 	/** Creates the asset picker widget for selecting a recording. */
 	TSharedRef<SWidget> CreateRecordingPicker()
 	{
 		FMenuBuilder MenuBuilder(true, nullptr);
 
 		IContentBrowserSingleton& ContentBrowser = FModuleManager::LoadModuleChecked<FContentBrowserModule>("ContentBrowser").Get();
-
+		
 		FAssetPickerConfig AssetPickerConfig;
 		{
 			AssetPickerConfig.SelectionMode = ESelectionMode::Single;
@@ -113,14 +199,12 @@ private:
 			AssetPickerConfig.bCanShowClasses = false;
 			AssetPickerConfig.bShowPathInColumnView = true;
 			AssetPickerConfig.bSortByPathInColumnView = false;
-			AssetPickerConfig.AssetShowWarningText = LOCTEXT("NoRecordings_Warning", "No Recordings Found");
+			AssetPickerConfig.AssetShowWarningText = GetNoAssetsWarningText();
 
 			AssetPickerConfig.bForceShowEngineContent = true;
 			AssetPickerConfig.bForceShowPluginContent = true;
 
-			AssetPickerConfig.Filter.ClassPaths.Add(ULiveLinkRecording::StaticClass()->GetClassPathName());
-			AssetPickerConfig.Filter.bRecursiveClasses = true;
-			AssetPickerConfig.Filter.bRecursivePaths = true;
+			AssetPickerConfig.Filter = MakeAssetFilter();
 			AssetPickerConfig.OnAssetDoubleClicked = FOnAssetSelected::CreateRaw(this, &SLiveLinkHubRecordingListView::OnImportRecording);
 		}
 
@@ -140,6 +224,24 @@ private:
 		return MenuBuilder.MakeWidget();
 	}
 
+	/** Create a filter for available recording assets. */
+	FARFilter MakeAssetFilter() const
+	{
+		FARFilter Filter;
+		Filter.ClassPaths.Add(ULiveLinkRecording::StaticClass()->GetClassPathName());
+		Filter.bRecursiveClasses = true;
+		Filter.bRecursivePaths = true;
+		// There shouldn't be recordings that exist in memory but not on disk. Necessary to properly register deleted assets.
+		Filter.bIncludeOnlyOnDiskAssets = true;
+		return Filter;
+	}
+
+	/** The text to display when no assets are found. */
+	static FText GetNoAssetsWarningText()
+	{
+		return LOCTEXT("NoRecordings_Warning", "No Recordings Found");
+	}
+	
 private:
 	/** Delegate used for noticing the hub that a recording was selected for playback. */
 	FOnImportRecording OnImportRecordingDelegate;
@@ -147,6 +249,13 @@ private:
 	FOnEject OnEjectDelegate;
 	/** Delegate used to determine a recording can be ejected. */
 	FCanEject OnCanEjectDelegate;
+	
+	/** Handle for when an asset is added to the asset registry. */
+	FDelegateHandle OnAssetAddedHandle;
+	/** Handle for when an asset is removed from the asset registry. */
+	FDelegateHandle OnAssetRemovedHandle;
+	/** True if there are recording assets that exist. */
+	mutable TOptional<bool> bAssetsAvailableCached;
 };
 
 #undef LOCTEXT_NAMESPACE /* LiveLinkHub.RecordingListView */
