LiveLinkHub - Allow editing preprocessors and translators inline

diff --git a/Engine/Plugins/Animation/LiveLinkHub/Source/LiveLinkHub/Private/Config/LiveLinkHubFileUtilities.cpp b/Engine/Plugins/Animation/LiveLinkHub/Source/LiveLinkHub/Private/Config/LiveLinkHubFileUtilities.cpp
--- a/Engine/Plugins/Animation/LiveLinkHub/Source/LiveLinkHub/Private/Config/LiveLinkHubFileUtilities.cpp
+++ b/Engine/Plugins/Animation/LiveLinkHub/Source/LiveLinkHub/Private/Config/LiveLinkHubFileUtilities.cpp
@@ -6,8 +6,9 @@
 #include "JsonObjectConverter.h"
 #include "LiveLinkHubLog.h"
 #include "Session/LiveLinkHubSessionData.h"
+#include "UObject/Package.h"
 
-void UE::LiveLinkHub::FileUtilities::Private::SaveConfig(const FLiveLinkHubPersistedSessionData& InConfigData, const FString& InFilePath)
+void UE::LiveLinkHub::FileUtilities::Private::SaveConfig(const ULiveLinkHubSessionData* InConfigData, const FString& InFilePath)
 {
 	if (ensure(!InFilePath.IsEmpty()))
 	{
@@ -24,7 +25,7 @@ void UE::LiveLinkHub::FileUtilities::Private::SaveConfig(const FLiveLinkHubPersi
 	}
 }
 
-TSharedPtr<FLiveLinkHubPersistedSessionData> UE::LiveLinkHub::FileUtilities::Private::LoadConfig(const FString& InFilePath)
+ULiveLinkHubSessionData* UE::LiveLinkHub::FileUtilities::Private::LoadConfig(const FString& InFilePath)
 {
 	if (IFileManager::Get().FileExists(*InFilePath))
 	{
@@ -52,31 +53,32 @@ TSharedPtr<FLiveLinkHubPersistedSessionData> UE::LiveLinkHub::FileUtilities::Pri
 	return nullptr;
 }
 
-TSharedPtr<FJsonObject> UE::LiveLinkHub::FileUtilities::Private::ToJson(const FLiveLinkHubPersistedSessionData& InConfigData)
+TSharedPtr<FJsonObject> UE::LiveLinkHub::FileUtilities::Private::ToJson(const ULiveLinkHubSessionData* InConfigData)
 {
 	TSharedRef<FJsonObject> JsonObject = MakeShared<FJsonObject>();
 	const TSharedPtr<FJsonValueNumber> NumberValue = MakeShared<FJsonValueNumber>(LiveLinkHubVersion);
 	JsonObject->SetField(JsonVersionKey, NumberValue);
 	
-	FJsonObjectConverter::UStructToJsonObject(FLiveLinkHubPersistedSessionData::StaticStruct(), &InConfigData, JsonObject);
+	FJsonObjectConverter::UStructToJsonObject(ULiveLinkHubSessionData::StaticClass(), InConfigData, JsonObject);
 	
 	return JsonObject;
 }
 
-TSharedPtr<FLiveLinkHubPersistedSessionData> UE::LiveLinkHub::FileUtilities::Private::FromJson(const TSharedPtr<FJsonObject>& InJsonObject)
+ULiveLinkHubSessionData* UE::LiveLinkHub::FileUtilities::Private::FromJson(const TSharedPtr<FJsonObject>& InJsonObject)
 {
 	if (InJsonObject.IsValid())
 	{
-		FLiveLinkHubPersistedSessionData OutConfigData;
+		ULiveLinkHubSessionData* OutConfigData = NewObject<ULiveLinkHubSessionData>(GetTransientPackage());
 		const bool bResult =
 			FJsonObjectConverter::JsonObjectToUStruct(InJsonObject.ToSharedRef(),
-				FLiveLinkHubPersistedSessionData::StaticStruct(), &OutConfigData);
+				ULiveLinkHubSessionData::StaticClass(), OutConfigData);
 
 		if (!bResult)
 		{
 			UE_LOG(LogLiveLinkHub, Error, TEXT("Could not convert from json to LiveLinkHubSessionData."))
 		}
-		return MakeShared<FLiveLinkHubPersistedSessionData>(MoveTemp(OutConfigData));
+
+		return OutConfigData;
 	}
 
 	return nullptr;
diff --git a/Engine/Plugins/Animation/LiveLinkHub/Source/LiveLinkHub/Private/Config/LiveLinkHubFileUtilities.h b/Engine/Plugins/Animation/LiveLinkHub/Source/LiveLinkHub/Private/Config/LiveLinkHubFileUtilities.h
--- a/Engine/Plugins/Animation/LiveLinkHub/Source/LiveLinkHub/Private/Config/LiveLinkHubFileUtilities.h
+++ b/Engine/Plugins/Animation/LiveLinkHub/Source/LiveLinkHub/Private/Config/LiveLinkHubFileUtilities.h
@@ -5,7 +5,7 @@
 #include "Containers/UnrealString.h"
 #include "Templates/SharedPointer.h"
 
-struct FLiveLinkHubPersistedSessionData;
+class ULiveLinkHubSessionData;
 class FJsonObject;
 
 namespace UE::LiveLinkHub::FileUtilities::Private
@@ -23,14 +23,14 @@ namespace UE::LiveLinkHub::FileUtilities::Private
 	const FString ConfigDescription = TEXT("Live Link Hub Config");
 
 	/** Save config data to disk. */
-	void SaveConfig(const FLiveLinkHubPersistedSessionData& InConfigData, const FString& InFilePath);
+	void SaveConfig(const ULiveLinkHubSessionData* InConfigData, const FString& InFilePath);
 
 	/** Load config data from disk. */
-	TSharedPtr<FLiveLinkHubPersistedSessionData> LoadConfig(const FString& InFilePath);
+	ULiveLinkHubSessionData* LoadConfig(const FString& InFilePath);
 
 	/** Convert config data to json. */
-	TSharedPtr<FJsonObject> ToJson(const FLiveLinkHubPersistedSessionData& InConfigData);
+	TSharedPtr<FJsonObject> ToJson(const ULiveLinkHubSessionData* InConfigData);
 
 	/** Convert config data from json. */
-	TSharedPtr<FLiveLinkHubPersistedSessionData> FromJson(const TSharedPtr<FJsonObject>& InJsonObject);
-}
\ No newline at end of file
+	ULiveLinkHubSessionData* FromJson(const TSharedPtr<FJsonObject>& InJsonObject);
+}
diff --git a/Engine/Plugins/Animation/LiveLinkHub/Source/LiveLinkHub/Private/Session/LiveLinkHubSession.h b/Engine/Plugins/Animation/LiveLinkHub/Source/LiveLinkHub/Private/Session/LiveLinkHubSession.h
--- a/Engine/Plugins/Animation/LiveLinkHub/Source/LiveLinkHub/Private/Session/LiveLinkHubSession.h
+++ b/Engine/Plugins/Animation/LiveLinkHub/Source/LiveLinkHub/Private/Session/LiveLinkHubSession.h
@@ -14,7 +14,10 @@
 #include "Misc/ScopeLock.h"
 #include "Modules/ModuleManager.h"
 #include "Subjects/LiveLinkHubSubjectSessionConfig.h"
+#include "UObject/Package.h"
 #include "UObject/StrongObjectPtr.h"
+#include "UObject/UObjectGlobals.h"
+
 
 #define LOCTEXT_NAMESPACE "LiveLinkHubSession"
 
@@ -35,6 +38,9 @@ public:
 	/** Get the configuration for a given subject. */
 	virtual TOptional<FLiveLinkHubSubjectProxy> GetSubjectConfig(const FLiveLinkSubjectKey& SubjectKey) const = 0;
 
+	/** Get the preprocessors and translator for a given subject. */
+	virtual ULiveLinkHubSubjectProcessors* GetSubjectProcessors(const FLiveLinkSubjectKey& SubjectKey) const = 0;
+
 	/** Change the outbound name of a subject for this session. */
 	virtual void RenameSubject(const FLiveLinkSubjectKey& SubjectKey, FName NewName) = 0;
 
@@ -66,16 +72,17 @@ public:
 	{
 		RegisterDelegates();
 
-		SessionData.SubjectsConfig.Initialize();
+		SessionData = TStrongObjectPtr<ULiveLinkHubSessionData>(NewObject<ULiveLinkHubSessionData>(GetTransientPackage()));
+		SessionData->SubjectsConfig->Initialize();
 	}
 
-	FLiveLinkHubSession(FLiveLinkHubSessionData InSessionData, FOnClientAddedToSession& OnClientAddedToSession, FOnClientRemovedFromSession& OnClientRemovedFromSession)
-		: SessionData(MoveTemp(InSessionData))
-		, OnClientAddedToSessionDelegate(OnClientAddedToSession)
+	FLiveLinkHubSession(ULiveLinkHubSessionData* InSessionData, FOnClientAddedToSession& OnClientAddedToSession, FOnClientRemovedFromSession& OnClientRemovedFromSession)
+		: OnClientAddedToSessionDelegate(OnClientAddedToSession)
 		, OnClientRemovedFromSessionDelegate(OnClientRemovedFromSession)
 	{
+		SessionData = TStrongObjectPtr<ULiveLinkHubSessionData>(InSessionData);
 		RegisterDelegates();
-		SessionData.SubjectsConfig.Initialize();
+		SessionData->SubjectsConfig->Initialize();
 	}
 
 	virtual ~FLiveLinkHubSession() override
@@ -85,25 +92,35 @@ public:
 
 	virtual void RenameSubject(const FLiveLinkSubjectKey& SubjectKey, FName NewName) override
 	{
-		FLiveLinkHubSubjectSessionConfig ConfigCopy;
+		TMap<FLiveLinkSubjectKey, FLiveLinkHubSubjectProxy> SubjectProxies;
 		{
 			FReadScopeLock Locker(SessionDataLock);
-			ConfigCopy = SessionData.SubjectsConfig;
+
+			SubjectProxies = SessionData->SubjectsConfig->SubjectProxies;
 		}
 
-		ConfigCopy.RenameSubject(SubjectKey, NewName);
+		if (FLiveLinkHubSubjectProxy* Proxy = SubjectProxies.Find(SubjectKey))
+		{
+			Proxy->SetOutboundName(NewName);
+		}
 
 		{
 			// Copied over in a different step to avoid acquiring the rw lock in a method called by RenameSubject
 			FWriteScopeLock Locker(SessionDataLock);
-			SessionData.SubjectsConfig = MoveTemp(ConfigCopy);
+			SessionData->SubjectsConfig->SubjectProxies = MoveTemp(SubjectProxies);
 		}
 	}
 
 	virtual TOptional<FLiveLinkHubSubjectProxy> GetSubjectConfig(const FLiveLinkSubjectKey& SubjectKey) const override
 	{
 		FReadScopeLock Locker(SessionDataLock);
-		return SessionData.SubjectsConfig.GetSubjectConfig(SubjectKey);
+		return SessionData->SubjectsConfig->GetSubjectConfig(SubjectKey);
+	}
+
+	virtual ULiveLinkHubSubjectProcessors* GetSubjectProcessors(const FLiveLinkSubjectKey& SubjectKey) const override
+	{
+		FReadScopeLock Locker(SessionDataLock);
+		return SessionData->SubjectsConfig->GetSubjectProcessors(SubjectKey);
 	}
 
 	virtual TArray<FLiveLinkHubClientId> GetSessionClients() const override
@@ -151,18 +168,18 @@ public:
 	virtual void SetPreProcessors(const FLiveLinkSubjectKey& SubjectKey, TConstArrayView<ULiveLinkFramePreProcessor*> PreProcessors) override
 	{
 		FWriteScopeLock Locker(SessionDataLock);
-		if (FLiveLinkHubSubjectProxy* SubjectProxy = SessionData.SubjectsConfig.SubjectProxies.Find(SubjectKey))
+		if (TObjectPtr<ULiveLinkHubSubjectProcessors>* SubjectProcessors = SessionData->SubjectsConfig->SubjectProcessors.Find(SubjectKey))
 		{
-			SubjectProxy->SetPreProcessors(PreProcessors);
+			(*SubjectProcessors)->PreProcessors = PreProcessors;
 		}
 	}
 
 	virtual void SetTranslator(const FLiveLinkSubjectKey& SubjectKey, ULiveLinkFrameTranslator* Translator) override
 	{
 		FWriteScopeLock Locker(SessionDataLock);
-		if (FLiveLinkHubSubjectProxy* SubjectProxy = SessionData.SubjectsConfig.SubjectProxies.Find(SubjectKey))
+		if (TObjectPtr<ULiveLinkHubSubjectProcessors>* SubjectProcessors = SessionData->SubjectsConfig->SubjectProcessors.Find(SubjectKey))
 		{
-			SubjectProxy->SetTranslator(Translator);
+			(*SubjectProcessors)->Translator = Translator;
 		}
 	}
 
@@ -217,16 +234,21 @@ private:
 	/** Handles updating the tree view when a subject is added. */
 	void OnSubjectAdded(const FLiveLinkSubjectKey& SubjectKey)
 	{
-		if (!SessionData.SubjectsConfig.SubjectProxies.Contains(SubjectKey))
+		if (!SessionData->SubjectsConfig->SubjectProxies.Contains(SubjectKey))
 		{
 			ILiveLinkClient& LiveLinkClient = IModularFeatures::Get().GetModularFeature<ILiveLinkClient>(ILiveLinkClient::ModularFeatureName);
 
 			FLiveLinkHubSubjectProxy SubjectSettings;
-			SubjectSettings.Initialize(SubjectKey, LiveLinkClient.GetSourceType(SubjectKey.Source).ToString(), Cast<ULiveLinkSubjectSettings>(LiveLinkClient.GetSubjectSettings(SubjectKey)));
+			SubjectSettings.Initialize(SubjectKey, LiveLinkClient.GetSourceType(SubjectKey.Source).ToString());
 
 			{
 				FWriteScopeLock Locker(SessionDataLock);
-				SessionData.SubjectsConfig.SubjectProxies.FindOrAdd(SubjectKey) = MoveTemp(SubjectSettings);
+				SessionData->SubjectsConfig->SubjectProxies.FindOrAdd(SubjectKey) = MoveTemp(SubjectSettings);
+
+				ULiveLinkHubSubjectProcessors* SubjectProcessors = NewObject<ULiveLinkHubSubjectProcessors>(SessionData->SubjectsConfig);
+				SubjectProcessors->Initialize(Cast<ULiveLinkSubjectSettings>(LiveLinkClient.GetSubjectSettings(SubjectKey)));
+
+				SessionData->SubjectsConfig->SubjectProcessors.FindOrAdd(SubjectKey) = SubjectProcessors;
 			}
 		}
 	}
@@ -248,7 +270,8 @@ private:
 	void OnSubjectRemoved(const FLiveLinkSubjectKey& SubjectKey)
 	{
 		FWriteScopeLock Locker(SessionDataLock);
-		SessionData.SubjectsConfig.SubjectProxies.Remove(SubjectKey);
+		SessionData->SubjectsConfig->SubjectProxies.Remove(SubjectKey);
+		SessionData->SubjectsConfig->SubjectProcessors.Remove(SubjectKey);
 	}
 
 private:
@@ -256,7 +279,7 @@ private:
 	TSet<FLiveLinkHubClientId> CachedSessionClients;
 
 	/** Holds data for this session. */
-	FLiveLinkHubSessionData SessionData;
+	TStrongObjectPtr<ULiveLinkHubSessionData> SessionData;
 
 	/** Delegate used to notice the hub about clients being added to this session. */
 	FOnClientAddedToSession& OnClientAddedToSessionDelegate;
diff --git a/Engine/Plugins/Animation/LiveLinkHub/Source/LiveLinkHub/Private/Session/LiveLinkHubSessionData.h b/Engine/Plugins/Animation/LiveLinkHub/Source/LiveLinkHub/Private/Session/LiveLinkHubSessionData.h
--- a/Engine/Plugins/Animation/LiveLinkHub/Source/LiveLinkHub/Private/Session/LiveLinkHubSessionData.h
+++ b/Engine/Plugins/Animation/LiveLinkHub/Source/LiveLinkHub/Private/Session/LiveLinkHubSessionData.h
@@ -4,6 +4,7 @@
 
 #include "Clients/LiveLinkHubUEClientInfo.h"
 #include "CoreTypes.h"
+#include "Engine/TimecodeProvider.h"
 #include "LiveLinkPresetTypes.h"
 #include "Misc/Guid.h"
 #include "Subjects/LiveLinkHubSubjectSessionConfig.h"
@@ -11,28 +12,26 @@
 
 #include "LiveLinkHubSessionData.generated.h"
 
-/** Live link hub session data which can be serialized to disk. */
-USTRUCT()
-struct FLiveLinkHubSessionData
-{
-	GENERATED_BODY()
-
-	UPROPERTY()
-	/** Subject configs for this session. */
-	FLiveLinkHubSubjectSessionConfig SubjectsConfig;
-};
-
 /** Live link hub session data that can be saved to disk. */
-USTRUCT()
-struct FLiveLinkHubPersistedSessionData : public FLiveLinkHubSessionData
+UCLASS()
+class ULiveLinkHubSessionData : public UObject
 {
+public:
 	GENERATED_BODY()
 
-	FLiveLinkHubPersistedSessionData() = default;
+	ULiveLinkHubSessionData()
+	{
+		SubjectsConfig = CreateDefaultSubobject<ULiveLinkHubSubjectSessionConfig>(TEXT("SubjectsConfig"));
+
+        if (!HasAnyFlags(RF_ClassDefaultObject | RF_ArchetypeObject))
+        {
+        	SubjectsConfig->Initialize();
+        }
+	}
 
-	FLiveLinkHubPersistedSessionData(FLiveLinkHubSessionData SessionData)
-		: FLiveLinkHubSessionData(MoveTemp(SessionData))
+	ULiveLinkHubSessionData(ULiveLinkHubSubjectSessionConfig* InSubjectsConfig)
 	{
+		SubjectsConfig = InSubjectsConfig;
 	}
 
 	/** Live link hub sources. */
@@ -50,4 +49,8 @@ struct FLiveLinkHubPersistedSessionData : public FLiveLinkHubSessionData
 	/** Timecode settings for the live link hub. */
 	UPROPERTY()
 	FLiveLinkHubTimecodeSettings TimecodeSettings;
+
+	/** Subject configs for this session. */
+	UPROPERTY(Instanced)
+	TObjectPtr<ULiveLinkHubSubjectSessionConfig> SubjectsConfig;
 };
diff --git a/Engine/Plugins/Animation/LiveLinkHub/Source/LiveLinkHub/Private/Session/LiveLinkHubSessionManager.h b/Engine/Plugins/Animation/LiveLinkHub/Source/LiveLinkHub/Private/Session/LiveLinkHubSessionManager.h
--- a/Engine/Plugins/Animation/LiveLinkHub/Source/LiveLinkHub/Private/Session/LiveLinkHubSessionManager.h
+++ b/Engine/Plugins/Animation/LiveLinkHub/Source/LiveLinkHub/Private/Session/LiveLinkHubSessionManager.h
@@ -141,29 +141,29 @@ public:
 			CurrentSessionPtr =  CurrentSession;
 		}
 
-		FLiveLinkHubPersistedSessionData LiveLinkHubSessionData = StaticCastSharedPtr<FLiveLinkHubSession>(CurrentSessionPtr)->SessionData;
+		ULiveLinkHubSessionData* LiveLinkHubSessionData = Cast<ULiveLinkHubSessionData>(StaticCastSharedPtr<FLiveLinkHubSession>(CurrentSessionPtr)->SessionData.Get());
 
 		TArray<FGuid> SourceGuids = LiveLinkHubClient->GetSources();
 		for (const FGuid& SourceGuid : SourceGuids)
 		{
-			LiveLinkHubSessionData.Sources.Add(LiveLinkHubClient->GetSourcePreset(SourceGuid, nullptr));
+			LiveLinkHubSessionData->Sources.Add(LiveLinkHubClient->GetSourcePreset(SourceGuid, nullptr));
 		}
 
 		TArray<FLiveLinkSubjectKey> Subjects = LiveLinkHubClient->GetSubjects(true, true);
 		for (const FLiveLinkSubjectKey& Subject : Subjects)
 		{
-			LiveLinkHubSessionData.Subjects.Add(LiveLinkHubClient->GetSubjectPreset(Subject, nullptr));
+			LiveLinkHubSessionData->Subjects.Add(LiveLinkHubClient->GetSubjectPreset(Subject, nullptr));
 		}
 
 		const TMap<FLiveLinkHubClientId, FLiveLinkHubUEClientInfo>& ClientMap = LiveLinkProvider->GetClientsMap();
 
 		for (const TTuple<FLiveLinkHubClientId, FLiveLinkHubUEClientInfo>& ClientKeyVal : ClientMap)
 		{
-			LiveLinkHubSessionData.Clients.Add(ClientKeyVal.Value);
+			LiveLinkHubSessionData->Clients.Add(ClientKeyVal.Value);
 		}
 
 		const FLiveLinkHubTimecodeSettings& TimecodeSettings = LiveLinkProvider->GetTimecodeSettings();
-		LiveLinkHubSessionData.TimecodeSettings = TimecodeSettings;
+		LiveLinkHubSessionData->TimecodeSettings = TimecodeSettings;
 		
 		if (!SavePath.IsEmpty())
 		{
@@ -222,7 +222,7 @@ private:
 		LastConfigPath = Path;
 		FEditorDirectories::Get().SetLastDirectory(ELastDirectory::GENERIC_OPEN, FPaths::GetPath(LastConfigPath));
 
-		const TSharedPtr<FLiveLinkHubPersistedSessionData> SessionData = UE::LiveLinkHub::FileUtilities::Private::LoadConfig(LastConfigPath);
+		ULiveLinkHubSessionData* SessionData = UE::LiveLinkHub::FileUtilities::Private::LoadConfig(LastConfigPath);
 
 		FLiveLinkHubClient* LiveLinkHubClient = static_cast<FLiveLinkHubClient*>(&IModularFeatures::Get().GetModularFeature<ILiveLinkClient>(ILiveLinkClient::ModularFeatureName));
 
@@ -232,7 +232,7 @@ private:
 		check(LiveLinkHubClient);
 		check(LiveLinkProvider);
 
-		if (SessionData.IsValid())
+		if (SessionData)
 		{
 			for (const FLiveLinkSourcePreset& SourcePreset : SessionData->Sources)
 			{
@@ -251,7 +251,7 @@ private:
 		TSharedPtr<FLiveLinkHubSession> CurrentSessionPtr;
 		{
 			FScopeLock Lock(&CurrentSessionCS);
-			CurrentSessionPtr = CurrentSession = MakeShared<FLiveLinkHubSession>(*SessionData, OnClientAddedToSessionDelegate, OnClientRemovedFromSessionDelegate);
+			CurrentSessionPtr = CurrentSession = MakeShared<FLiveLinkHubSession>(SessionData, OnClientAddedToSessionDelegate, OnClientRemovedFromSessionDelegate);
 		}
 
 		for (FLiveLinkHubUEClientInfo& Client : SessionData->Clients)
diff --git a/Engine/Plugins/Animation/LiveLinkHub/Source/LiveLinkHub/Private/Subjects/LiveLinkHubSubjectController.h b/Engine/Plugins/Animation/LiveLinkHub/Source/LiveLinkHub/Private/Subjects/LiveLinkHubSubjectController.h
--- a/Engine/Plugins/Animation/LiveLinkHub/Source/LiveLinkHub/Private/Subjects/LiveLinkHubSubjectController.h
+++ b/Engine/Plugins/Animation/LiveLinkHub/Source/LiveLinkHub/Private/Subjects/LiveLinkHubSubjectController.h
@@ -56,26 +56,20 @@ public:
 		}
 	}
 
-	void OnSubjectProcessorModified(const FLiveLinkSubjectKey& SubjectKey, const TArray<TSubclassOf<ULiveLinkFramePreProcessor>>& UpdatedPreprocessors, const TSubclassOf<ULiveLinkFrameTranslator> UpdatedTranslator) const
+	void OnSubjectProcessorModified(const FLiveLinkSubjectKey& SubjectKey, const TArray<ULiveLinkFramePreProcessor*>& UpdatedPreprocessors, ULiveLinkFrameTranslator* UpdatedTranslator) const
 	{
 		FLiveLinkHubClient& LiveLinkClient = static_cast<FLiveLinkHubClient&>(IModularFeatures::Get().GetModularFeature<FLiveLinkClient>(ILiveLinkClient::ModularFeatureName));
 		ULiveLinkSubjectSettings* Settings = Cast<ULiveLinkSubjectSettings>(LiveLinkClient.GetSubjectSettings(SubjectKey));
 
-		Settings->PreProcessors.Reset(UpdatedPreprocessors.Num());
-		for (TSubclassOf<ULiveLinkFramePreProcessor> PreProcessor : UpdatedPreprocessors)
-		{
-			Settings->PreProcessors.Add(PreProcessor.GetDefaultObject());
-		}
-
+		Settings->PreProcessors = UpdatedPreprocessors;
 		Settings->Translators.Reset();
 
 		if (UpdatedTranslator)
 		{
-			Settings->Translators.Add(UpdatedTranslator.GetDefaultObject());
+			Settings->Translators.Add(UpdatedTranslator);
 		}
 
 		// todo: If we are removing the translator, we will need to do additional handling to restore the original static data that was overriden
-
 		if (Settings->ValidateProcessors())
 		{
 			// Apply to the underlying data from the session.
diff --git a/Engine/Plugins/Animation/LiveLinkHub/Source/LiveLinkHub/Private/Subjects/LiveLinkHubSubjectModel.h b/Engine/Plugins/Animation/LiveLinkHub/Source/LiveLinkHub/Private/Subjects/LiveLinkHubSubjectModel.h
--- a/Engine/Plugins/Animation/LiveLinkHub/Source/LiveLinkHub/Private/Subjects/LiveLinkHubSubjectModel.h
+++ b/Engine/Plugins/Animation/LiveLinkHub/Source/LiveLinkHub/Private/Subjects/LiveLinkHubSubjectModel.h
@@ -17,6 +17,9 @@ public:
 
 	/** Get the settings for a livelinkhub subject. */
 	virtual TOptional<FLiveLinkHubSubjectProxy> GetSubjectConfig(const FLiveLinkSubjectKey& InSubject) const = 0;
+
+	/** Get the preprocessors and translator for a livelinkhub subject. */
+	virtual ULiveLinkHubSubjectProcessors* GetSubjectProcessors(const FLiveLinkSubjectKey& InSubject) const = 0;
 };
 
 /** Implementation of the ILiveLinkHubSubjectModel. */
@@ -33,4 +36,15 @@ public:
 		}
 		return {};
 	}
+
+	/** Get the settings for a livelinkhub subject. */
+	virtual ULiveLinkHubSubjectProcessors* GetSubjectProcessors(const FLiveLinkSubjectKey& InSubject) const override
+	{
+		const FLiveLinkHubModule& LiveLinkHubModule = FModuleManager::Get().GetModuleChecked<FLiveLinkHubModule>("LiveLinkHub");
+		if (const TSharedPtr<ILiveLinkHubSession> Session = LiveLinkHubModule.GetSessionManager()->GetCurrentSession())
+		{
+			return Session->GetSubjectProcessors(InSubject);
+		}
+		return nullptr;
+	}
 };
diff --git a/Engine/Plugins/Animation/LiveLinkHub/Source/LiveLinkHub/Private/Subjects/LiveLinkHubSubjectSessionConfig.cpp b/Engine/Plugins/Animation/LiveLinkHub/Source/LiveLinkHub/Private/Subjects/LiveLinkHubSubjectSessionConfig.cpp
--- a/Engine/Plugins/Animation/LiveLinkHub/Source/LiveLinkHub/Private/Subjects/LiveLinkHubSubjectSessionConfig.cpp
+++ b/Engine/Plugins/Animation/LiveLinkHub/Source/LiveLinkHub/Private/Subjects/LiveLinkHubSubjectSessionConfig.cpp
@@ -6,15 +6,16 @@
 #include "Clients/LiveLinkHubProvider.h"
 #include "Features/IModularFeatures.h"
 #include "LiveLinkClient.h"
+#include "LiveLinkHubClient.h"
 #include "LiveLinkHubModule.h"
 #include "LiveLinkSubjectSettings.h"
 #include "Modules/ModuleManager.h"
 #include "UObject/StrongObjectPtr.h"
 
 
-void FLiveLinkHubSubjectSessionConfig::Initialize()
+void ULiveLinkHubSubjectSessionConfig::Initialize()
 {
-	FLiveLinkClient& LiveLinkClient = IModularFeatures::Get().GetModularFeature<FLiveLinkClient>(ILiveLinkClient::ModularFeatureName);
+	FLiveLinkHubClient& LiveLinkClient = IModularFeatures::Get().GetModularFeature<FLiveLinkHubClient>(ILiveLinkClient::ModularFeatureName);
 
 	constexpr bool bIncludeDisabledSubject = true;
 	constexpr bool bIncludeVirtualSubject = true;
@@ -22,15 +23,40 @@ void FLiveLinkHubSubjectSessionConfig::Initialize()
 	for (const FLiveLinkSubjectKey& SubjectKey : LiveLinkClient.GetSubjects(bIncludeDisabledSubject, bIncludeVirtualSubject))
 	{
 		FLiveLinkHubSubjectProxy SubjectSettings;
+		SubjectSettings.Initialize(SubjectKey, LiveLinkClient.GetSourceType(SubjectKey.Source).ToString());
 
+		SubjectProxies.Add(SubjectKey, MoveTemp(SubjectSettings));
+
+		ULiveLinkHubSubjectProcessors* Processors = nullptr;
+		if (TObjectPtr<ULiveLinkHubSubjectProcessors>* ProcessorsPtr = SubjectProcessors.Find(SubjectKey))
+		{
+			Processors = *ProcessorsPtr;
+		}
+		else
+		{
+			Processors = NewObject<ULiveLinkHubSubjectProcessors>(this);
+			SubjectProcessors.Add(SubjectKey, Processors);
+		}
+
+		// Apply restored settings to the subject.
 		ULiveLinkSubjectSettings* Settings = Cast<ULiveLinkSubjectSettings>(LiveLinkClient.GetSubjectSettings(SubjectKey));
-		SubjectSettings.Initialize(SubjectKey, LiveLinkClient.GetSourceType(SubjectKey.Source).ToString(), Settings);
 
-		SubjectProxies.Add(SubjectKey, MoveTemp(SubjectSettings));
+		Settings->PreProcessors = Processors->PreProcessors;
+		Settings->Translators.Reset();
+
+		if (Processors->Translator)
+		{
+			Settings->Translators.Add(Processors->Translator);
+		}
+
+		if (Settings->ValidateProcessors())
+		{
+			LiveLinkClient.CacheSubjectSettings(SubjectKey, Settings);
+		}
 	}
 }
 
-TOptional<FLiveLinkHubSubjectProxy> FLiveLinkHubSubjectSessionConfig::GetSubjectConfig(const FLiveLinkSubjectKey& InSubject) const
+TOptional<FLiveLinkHubSubjectProxy> ULiveLinkHubSubjectSessionConfig::GetSubjectConfig(const FLiveLinkSubjectKey& InSubject) const
 {
 	TOptional<FLiveLinkHubSubjectProxy> Settings;
 
@@ -42,7 +68,18 @@ TOptional<FLiveLinkHubSubjectProxy> FLiveLinkHubSubjectSessionConfig::GetSubject
     return Settings;
 }
 
-void FLiveLinkHubSubjectSessionConfig::RenameSubject(const FLiveLinkSubjectKey& SubjectKey, FName NewName)
+ULiveLinkHubSubjectProcessors* ULiveLinkHubSubjectSessionConfig::GetSubjectProcessors(const FLiveLinkSubjectKey& InSubject)
+{
+	if (TObjectPtr<ULiveLinkHubSubjectProcessors>* SettingsPtr = SubjectProcessors.Find(InSubject))
+	{
+		return *SettingsPtr;
+	}
+
+	return nullptr;
+}
+
+
+void ULiveLinkHubSubjectSessionConfig::RenameSubject(const FLiveLinkSubjectKey& SubjectKey, FName NewName)
 {
 	if (FLiveLinkHubSubjectProxy* Proxy = SubjectProxies.Find(SubjectKey))
 	{
@@ -50,24 +87,12 @@ void FLiveLinkHubSubjectSessionConfig::RenameSubject(const FLiveLinkSubjectKey&
 	}
 }
 
-void FLiveLinkHubSubjectProxy::Initialize(const FLiveLinkSubjectKey& InSubjectKey, FString InSource, ULiveLinkSubjectSettings* InSubjectSettings)
+void FLiveLinkHubSubjectProxy::Initialize(const FLiveLinkSubjectKey& InSubjectKey, FString InSource)
 {
 	SubjectName = InSubjectKey.SubjectName.Name.ToString();
 	SubjectKey = InSubjectKey;
 	OutboundName = InSubjectKey.SubjectName.Name.ToString();
 	Source = MoveTemp(InSource);
-
-	PreProcessors.Reset(InSubjectSettings->PreProcessors.Num());
-	for (const ULiveLinkFramePreProcessor* PreProcessor : InSubjectSettings->PreProcessors)
-	{
-		PreProcessors.Add(PreProcessor->GetClass());
-	}
-
-	Translator = nullptr;
-	if (InSubjectSettings->Translators.Num() && InSubjectSettings->Translators[0])
-	{
-		Translator = InSubjectSettings->Translators[0]->GetClass();
-	}
 }
 
 FName FLiveLinkHubSubjectProxy::GetOutboundName() const
@@ -118,3 +143,14 @@ void FLiveLinkHubSubjectProxy::NotifyRename()
 		}
 	}
 }
+
+void ULiveLinkHubSubjectProcessors::Initialize(ULiveLinkSubjectSettings* InSubjectSettings)
+{
+	PreProcessors.Append(InSubjectSettings->PreProcessors);
+
+	Translator = nullptr;
+	if (InSubjectSettings->Translators.Num() && InSubjectSettings->Translators[0])
+	{
+		Translator = InSubjectSettings->Translators[0];
+	}
+}
diff --git a/Engine/Plugins/Animation/LiveLinkHub/Source/LiveLinkHub/Private/Subjects/LiveLinkHubSubjectSessionConfig.h b/Engine/Plugins/Animation/LiveLinkHub/Source/LiveLinkHub/Private/Subjects/LiveLinkHubSubjectSessionConfig.h
--- a/Engine/Plugins/Animation/LiveLinkHub/Source/LiveLinkHub/Private/Subjects/LiveLinkHubSubjectSessionConfig.h
+++ b/Engine/Plugins/Animation/LiveLinkHub/Source/LiveLinkHub/Private/Subjects/LiveLinkHubSubjectSessionConfig.h
@@ -21,7 +21,7 @@ public:
 	GENERATED_BODY()
 
 	/** Initialize from a subject key. */
-	void Initialize(const FLiveLinkSubjectKey& InSubjectKey, FString InSource, ULiveLinkSubjectSettings* InSubjectSettings);
+	void Initialize(const FLiveLinkSubjectKey& InSubjectKey, FString InSource);
 
 	/** Get the outbound name for this subject, allows  */
 	FName GetOutboundName() const;
@@ -29,26 +29,6 @@ public:
 	/** Change the outbound name for this subject proxy. */
 	void SetOutboundName(FName NewName);
 
-	/** Set the preprocessors */
-	void SetPreProcessors(TConstArrayView<ULiveLinkFramePreProcessor*> FramePreprocessors)
-	{
-		PreProcessors.Reset(FramePreprocessors.Num());
-		Algo::TransformIf(FramePreprocessors, PreProcessors, [](const ULiveLinkFramePreProcessor* PreProcessor) { return PreProcessor; }, [](const ULiveLinkFramePreProcessor* PreProcessor){ return PreProcessor->GetClass(); });
-	}
-
-	/** Set the frame translator */
-	void SetTranslator(ULiveLinkFrameTranslator* InFrameTranslator)
-	{
-		if (InFrameTranslator)
-		{
-			Translator = InFrameTranslator->GetClass();
-		}
-		else
-		{
-			Translator = nullptr;
-		}
-	}
-
 private:
 	/** Notify clients that the old subject should be deleted and replaced with a new one with the updated name. */
 	void NotifyRename();
@@ -66,14 +46,6 @@ private:
 	UPROPERTY(VisibleAnywhere, Category = "Subject Details")
 	FString Source;
 
-	/** Preprocessors to apply to the livelink data coming through the hub. */
-	UPROPERTY(config, EditAnywhere, Category="Subject Details")
-	TArray<TSubclassOf<ULiveLinkFramePreProcessor>> PreProcessors;
-
-	/** Translators used to convert livelink data to a different role. */
-	UPROPERTY(config, EditAnywhere, Category="Subject Details")
-	TSubclassOf<ULiveLinkFrameTranslator> Translator = nullptr;
-
 	/** SubjectKey for this subject, */
 	UPROPERTY()
 	FLiveLinkSubjectKey SubjectKey;
@@ -90,9 +62,26 @@ private:
 	friend class SLiveLinkHubSubjectView;
 };
 
+UCLASS()
+class ULiveLinkHubSubjectProcessors : public UObject
+{
+public:
+	GENERATED_BODY()
+
+	void Initialize(ULiveLinkSubjectSettings* InSubjectSettings);
+
+	/** Preprocessors to apply to the livelink data coming through the hub. */
+	UPROPERTY(EditAnywhere, Category="Processors", Instanced)
+	TArray<TObjectPtr<ULiveLinkFramePreProcessor>> PreProcessors;
+
+	/** Translators used to convert livelink data to a different role. */
+	UPROPERTY(EditAnywhere, Category="Processors", Instanced)
+	TObjectPtr<ULiveLinkFrameTranslator> Translator;
+};
+
 /** Config pertaining to livelink hub subjects for a given session. */
-USTRUCT()
-struct FLiveLinkHubSubjectSessionConfig
+UCLASS()
+class ULiveLinkHubSubjectSessionConfig : public UObject
 {
 public:
 	GENERATED_BODY()
@@ -102,13 +91,23 @@ public:
 	/** Get the config for a livelinkhub subject. */
 	TOptional<FLiveLinkHubSubjectProxy> GetSubjectConfig(const FLiveLinkSubjectKey& InSubject) const;
 
+	/** Get the subject processors. */
+	ULiveLinkHubSubjectProcessors* GetSubjectProcessors(const FLiveLinkSubjectKey& InSubject);
+
 	/** Change the outbound name of a subject for the current session. */
 	void RenameSubject(const FLiveLinkSubjectKey& SubjectKey, FName NewName);
 
 private:
 	/** Settings for subjects displayed in the livelink hub. */
-	UPROPERTY()
+	UPROPERTY(Transient)
 	TMap<FLiveLinkSubjectKey, FLiveLinkHubSubjectProxy> SubjectProxies;
 
+	/** 
+	 * Processor settings held in a UObject to allow creating and editing them inline. 
+	 * Note that the other parts of the session config are not held in a UObject since they can be accessed at any time on other threads.
+	 */
+	UPROPERTY(Instanced)
+	TMap<FLiveLinkSubjectKey, TObjectPtr<ULiveLinkHubSubjectProcessors>> SubjectProcessors;
+
 	friend class FLiveLinkHubSession;
 };
diff --git a/Engine/Plugins/Animation/LiveLinkHub/Source/LiveLinkHub/Private/Subjects/SLiveLinkHubSubjectView.h b/Engine/Plugins/Animation/LiveLinkHub/Source/LiveLinkHub/Private/Subjects/SLiveLinkHubSubjectView.h
--- a/Engine/Plugins/Animation/LiveLinkHub/Source/LiveLinkHub/Private/Subjects/SLiveLinkHubSubjectView.h
+++ b/Engine/Plugins/Animation/LiveLinkHub/Source/LiveLinkHub/Private/Subjects/SLiveLinkHubSubjectView.h
@@ -20,7 +20,7 @@
 #include "Widgets/SNullWidget.h"
 
 DECLARE_DELEGATE_TwoParams(FOnRenameLiveLinkHubSubject, const FLiveLinkSubjectKey& /*SubjectKey*/, FName /*NewName*/);
-DECLARE_DELEGATE_ThreeParams(FOnSubjectProcessorModified, const FLiveLinkSubjectKey& /*SubjectKey*/,  const TArray<TSubclassOf<ULiveLinkFramePreProcessor>>& /*UpdatedPreProcessors*/, TSubclassOf<ULiveLinkFrameTranslator> /*UpdatedTranslator*/);
+DECLARE_DELEGATE_ThreeParams(FOnSubjectProcessorModified, const FLiveLinkSubjectKey& /*SubjectKey*/,  const TArray<ULiveLinkFramePreProcessor*>& /*UpdatedPreProcessors*/, ULiveLinkFrameTranslator* /*UpdatedTranslator*/);
 
 /**
  * Provides the UI that displays information about a livelink hub subject.
@@ -42,6 +42,9 @@ public:
 		OnRenameSubjectDelegate = InArgs._OnRenameSubject;
 		OnProcessorModifiedDelegate = InArgs._OnProcessorModified;
 
+		ILiveLinkClient& LiveLinkClient = static_cast<ILiveLinkClient&>(IModularFeatures::Get().GetModularFeature<FLiveLinkClient>(ILiveLinkClient::ModularFeatureName));
+		LiveLinkClient.OnLiveLinkSubjectRemoved().AddSP(this, &SLiveLinkHubSubjectView::OnSubjectRemoved);
+
 		FDetailsViewArgs DetailsViewArgs;
 		DetailsViewArgs.bUpdatesFromSelection = false;
 		DetailsViewArgs.bLockable = false;
@@ -60,6 +63,9 @@ public:
 		SettingsDetailsView = PropertyEditorModule.CreateStructureDetailView(DetailsViewArgs, StructureDetailsArgs, SubjectData);
 		SettingsDetailsView->GetOnFinishedChangingPropertiesDelegate().AddSP(this, &SLiveLinkHubSubjectView::OnSubjectPropertyModified);
 
+		SettingsObjectDetailsView = PropertyEditorModule.CreateDetailView(DetailsViewArgs);
+		SettingsObjectDetailsView->OnFinishedChangingProperties().AddSP(this, &SLiveLinkHubSubjectView::OnSubjectPropertyModified);
+
 		ChildSlot
 		[
 			SNew(SVerticalBox)
@@ -68,14 +74,29 @@ public:
 			[
 				SettingsDetailsView->GetWidget().ToSharedRef()
 			]
+			+ SVerticalBox::Slot()
+			.FillHeight(1.0f)
+			[
+				SettingsObjectDetailsView.ToSharedRef()
+			]
 		];
 	}
 
+	~SLiveLinkHubSubjectView() override
+	{
+		if (IModularFeatures::Get().IsModularFeatureAvailable(ILiveLinkClient::ModularFeatureName))
+		{
+			ILiveLinkClient& LiveLinkClient = static_cast<ILiveLinkClient&>(IModularFeatures::Get().GetModularFeature<FLiveLinkClient>(ILiveLinkClient::ModularFeatureName));
+			LiveLinkClient.OnLiveLinkSubjectRemoved().RemoveAll(this);
+		}
+	}
+
 	/** Clear the subject details. */
-	void RefreshSubjectDetails(const TSharedRef<ILiveLinkHubSession>& ActiveSession)
+	void RefreshSubjectDetails(const TSharedPtr<ILiveLinkHubSession>& ActiveSession)
 	{
 		SubjectData->Reset();
 		SettingsDetailsView->SetStructureData(nullptr);
+		SettingsObjectDetailsView->SetObject(nullptr);
 
 		SetSubject(SubjectKey);
 	}
@@ -91,6 +112,12 @@ public:
 			SubjectData->InitializeAs<FLiveLinkHubSubjectProxy>(*SubjectProxy);
 			SettingsDetailsView->SetStructureData(SubjectData);
 		}
+
+		if (ULiveLinkHubSubjectProcessors* SubjectProxy = SubjectModel->GetSubjectProcessors(InSubjectKey))
+		{
+			SettingsObjectDetailsView->SetObject(SubjectProxy);
+			WeakProcessors = SubjectProxy;
+		}
 	}
 
 	/** Get the outbound name for a subject. */
@@ -114,18 +141,33 @@ public:
 				{
 					OnRenameSubjectDelegate.ExecuteIfBound(SubjectKey, Proxy->GetOutboundName());
 				}
-				else if (PropertyChangedEvent.GetPropertyName() == GET_MEMBER_NAME_CHECKED(FLiveLinkHubSubjectProxy, PreProcessors)
-					|| PropertyChangedEvent.GetPropertyName() == GET_MEMBER_NAME_CHECKED(FLiveLinkHubSubjectProxy, Translator))
-				{
-					OnProcessorModifiedDelegate.ExecuteIfBound(SubjectKey, Proxy->PreProcessors, Proxy->Translator);
-				}
+			}
+		}
+
+		if (ULiveLinkHubSubjectProcessors* SubjectProcessors = WeakProcessors.Get())
+		{
+			if (PropertyChangedEvent.GetPropertyName() == GET_MEMBER_NAME_CHECKED(ULiveLinkHubSubjectProcessors, PreProcessors)
+				|| PropertyChangedEvent.GetPropertyName() == GET_MEMBER_NAME_CHECKED(ULiveLinkHubSubjectProcessors, Translator))
+			{
+				OnProcessorModifiedDelegate.ExecuteIfBound(SubjectKey, SubjectProcessors->PreProcessors, SubjectProcessors->Translator);
 			}
 		}
 	}
 
+	void OnSubjectRemoved(FLiveLinkSubjectKey InSubjectKey)
+	{
+		if (InSubjectKey == SubjectKey)
+		{
+			SubjectKey = FLiveLinkSubjectKey();
+			RefreshSubjectDetails(nullptr);
+		}
+	}
+
 private:
 	/** Details for the selected subject. */
 	TSharedPtr<IStructureDetailsView> SettingsDetailsView;
+	/** Details for the selected subject. */
+	TSharedPtr<IDetailsView> SettingsObjectDetailsView;
 	/** Subject being shown. */
 	FLiveLinkSubjectKey SubjectKey;
 	/** Model that holds the data for a given subject. */
@@ -136,4 +178,6 @@ private:
 	FOnRenameLiveLinkHubSubject OnRenameSubjectDelegate;
 	/** Delegate called when a translator or preprocessor is modified by the user. */
 	FOnSubjectProcessorModified OnProcessorModifiedDelegate;
+	/** Subject proxy object that holds preprocessor data. */
+	TWeakObjectPtr<ULiveLinkHubSubjectProcessors> WeakProcessors;
 };
