Fix static analysis warning in livelink hub

diff --git a/Engine/Plugins/Animation/LiveLinkHub/Source/LiveLinkHub/Private/Session/LiveLinkHubSessionManager.h b/Engine/Plugins/Animation/LiveLinkHub/Source/LiveLinkHub/Private/Session/LiveLinkHubSessionManager.h
--- a/Engine/Plugins/Animation/LiveLinkHub/Source/LiveLinkHub/Private/Session/LiveLinkHubSessionManager.h
+++ b/Engine/Plugins/Animation/LiveLinkHub/Source/LiveLinkHub/Private/Session/LiveLinkHubSessionManager.h
@@ -222,44 +222,44 @@ private:
 		LastConfigPath = Path;
 		FEditorDirectories::Get().SetLastDirectory(ELastDirectory::GENERIC_OPEN, FPaths::GetPath(LastConfigPath));
 
-		ULiveLinkHubSessionData* SessionData = UE::LiveLinkHub::FileUtilities::Private::LoadConfig(LastConfigPath);
-
-		FLiveLinkHubClient* LiveLinkHubClient = static_cast<FLiveLinkHubClient*>(&IModularFeatures::Get().GetModularFeature<ILiveLinkClient>(ILiveLinkClient::ModularFeatureName));
+		if (ULiveLinkHubSessionData* SessionData = UE::LiveLinkHub::FileUtilities::Private::LoadConfig(LastConfigPath))
+		{
+			FLiveLinkHubClient* LiveLinkHubClient = static_cast<FLiveLinkHubClient*>(&IModularFeatures::Get().GetModularFeature<ILiveLinkClient>(ILiveLinkClient::ModularFeatureName));
 
-		const FLiveLinkHubModule& LiveLinkHubModule = FModuleManager::Get().GetModuleChecked<FLiveLinkHubModule>("LiveLinkHub");
-		const TSharedPtr<FLiveLinkHubProvider> LiveLinkProvider = LiveLinkHubModule.GetLiveLinkProvider();
+			const FLiveLinkHubModule& LiveLinkHubModule = FModuleManager::Get().GetModuleChecked<FLiveLinkHubModule>("LiveLinkHub");
+			const TSharedPtr<FLiveLinkHubProvider> LiveLinkProvider = LiveLinkHubModule.GetLiveLinkProvider();
 
-		check(LiveLinkHubClient);
-		check(LiveLinkProvider);
+			check(LiveLinkHubClient);
+			check(LiveLinkProvider);
 
-		if (SessionData)
-		{
-			for (const FLiveLinkSourcePreset& SourcePreset : SessionData->Sources)
+			if (SessionData)
 			{
-				LiveLinkHubClient->CreateSource(SourcePreset);
+				for (const FLiveLinkSourcePreset& SourcePreset : SessionData->Sources)
+				{
+					LiveLinkHubClient->CreateSource(SourcePreset);
+				}
+
+				for (const FLiveLinkSubjectPreset& SubjectPreset : SessionData->Subjects)
+				{
+					LiveLinkHubClient->CreateSubject(SubjectPreset);
+				}
 			}
 
-			for (const FLiveLinkSubjectPreset& SubjectPreset : SessionData->Subjects)
+			SessionData->TimecodeSettings.AssignTimecodeSettingsAsProviderToEngine();
+			LiveLinkProvider->SetTimecodeSettings(SessionData->TimecodeSettings);
+
+			TSharedPtr<FLiveLinkHubSession> CurrentSessionPtr;
 			{
-				LiveLinkHubClient->CreateSubject(SubjectPreset);
+				FScopeLock Lock(&CurrentSessionCS);
+				CurrentSessionPtr = CurrentSession = MakeShared<FLiveLinkHubSession>(SessionData, OnClientAddedToSessionDelegate, OnClientRemovedFromSessionDelegate);
 			}
-		}
 
-		SessionData->TimecodeSettings.AssignTimecodeSettingsAsProviderToEngine();
-		LiveLinkProvider->SetTimecodeSettings(SessionData->TimecodeSettings);
-
-		TSharedPtr<FLiveLinkHubSession> CurrentSessionPtr;
-		{
-			FScopeLock Lock(&CurrentSessionCS);
-			CurrentSessionPtr = CurrentSession = MakeShared<FLiveLinkHubSession>(SessionData, OnClientAddedToSessionDelegate, OnClientRemovedFromSessionDelegate);
-		}
-
-		for (FLiveLinkHubUEClientInfo& Client : SessionData->Clients)
-		{
-			CurrentSessionPtr->AddRestoredClient(Client);
+			for (FLiveLinkHubUEClientInfo& Client : SessionData->Clients)
+			{
+				CurrentSessionPtr->AddRestoredClient(Client);
+			}
+			OnActiveSessionChangedDelegate.Broadcast(CurrentSessionPtr.ToSharedRef());
 		}
-
-		OnActiveSessionChangedDelegate.Broadcast(CurrentSessionPtr.ToSharedRef());
 	}
 
 	/** Clear the hub data contained in the current session, resetting the hub to its default state. */
