Remove #pragma once in .cpp file (Horde Issue 488209)

diff --git a/Engine/Plugins/Animation/LiveLinkHub/Source/LiveLinkHub/Private/LiveLinkHubInputProcessor.cpp b/Engine/Plugins/Animation/LiveLinkHub/Source/LiveLinkHub/Private/LiveLinkHubInputProcessor.cpp
--- a/Engine/Plugins/Animation/LiveLinkHub/Source/LiveLinkHub/Private/LiveLinkHubInputProcessor.cpp
+++ b/Engine/Plugins/Animation/LiveLinkHub/Source/LiveLinkHub/Private/LiveLinkHubInputProcessor.cpp
@@ -1,5 +1,4 @@
 // Copyright Epic Games, Inc. All Rights Reserved.
-#pragma once
 
 #include "LiveLinkHubInputProcessor.h"
 #include "Input/Events.h"
