LiveLinkHub - Add livelink hub subject settings

These settings will be used to replace the subject session data and subject proxies in a subsequent CL.

diff --git a/Engine/Plugins/Animation/LiveLink/Source/LiveLink/Private/LiveLinkSettings.cpp b/Engine/Plugins/Animation/LiveLink/Source/LiveLink/Private/LiveLinkSettings.cpp
--- a/Engine/Plugins/Animation/LiveLink/Source/LiveLink/Private/LiveLinkSettings.cpp
+++ b/Engine/Plugins/Animation/LiveLink/Source/LiveLink/Private/LiveLinkSettings.cpp
@@ -7,9 +7,10 @@
 
 #include UE_INLINE_GENERATED_CPP_BY_NAME(LiveLinkSettings)
 
-FLiveLinkRoleProjectSetting::FLiveLinkRoleProjectSetting()
-	: SettingClass(ULiveLinkSubjectSettings::StaticClass())
-{}
+FLiveLinkRoleProjectSetting::FLiveLinkRoleProjectSetting(TSubclassOf<ULiveLinkSubjectSettings> DefaultSettingsClass)
+	: SettingClass(DefaultSettingsClass)
+{
+}
 
 
 ULiveLinkSettings::ULiveLinkSettings()
@@ -34,7 +35,14 @@ FLiveLinkRoleProjectSetting ULiveLinkSettings::GetDefaultSettingForRole(TSubclas
 	{
 		return DefaultRoleSettings[IndexOf];
 	}
-	FLiveLinkRoleProjectSetting Result;
+
+	TSubclassOf<ULiveLinkSubjectSettings> DefaultClass = ULiveLinkSubjectSettings::StaticClass();
+	if (UClass* Class = DefaultSettingsClass.TryLoadClass<ULiveLinkSubjectSettings>())
+	{
+		DefaultClass = Class;
+	}
+
+	FLiveLinkRoleProjectSetting Result{MoveTemp(DefaultClass)};
 	Result.Role = Role;
 	return Result;
 }
@@ -59,4 +67,3 @@ PRAGMA_DISABLE_DEPRECATION_WARNINGS
 PRAGMA_ENABLE_DEPRECATION_WARNINGS
 #endif //WITH_EDITOR
 }
-
diff --git a/Engine/Plugins/Animation/LiveLink/Source/LiveLink/Public/LiveLinkSettings.h b/Engine/Plugins/Animation/LiveLink/Source/LiveLink/Public/LiveLinkSettings.h
--- a/Engine/Plugins/Animation/LiveLink/Source/LiveLink/Public/LiveLinkSettings.h
+++ b/Engine/Plugins/Animation/LiveLink/Source/LiveLink/Public/LiveLinkSettings.h
@@ -25,7 +25,8 @@ struct LIVELINK_API FLiveLinkRoleProjectSetting
 	GENERATED_BODY()
 
 public:
-	FLiveLinkRoleProjectSetting();
+	FLiveLinkRoleProjectSetting() = default;
+	FLiveLinkRoleProjectSetting(TSubclassOf<ULiveLinkSubjectSettings> DefaultSettingsClass);
 
 public:
 	/** The role of the current setting. */
@@ -78,6 +79,10 @@ protected:
 	UPROPERTY(config, EditAnywhere, Category="LiveLink")
 	TArray<FLiveLinkRoleProjectSetting> DefaultRoleSettings;
 
+	/** When a settings class is not speficied for a role, this settings class will be used. */
+	UPROPERTY(config)
+	FSoftClassPath DefaultSettingsClass;
+
 public:
 	/** The interpolation class to use for new Subjects if no specific settings we set for the Subject's role. */
 	UPROPERTY(config)
@@ -156,12 +161,3 @@ public:
 	double GetMessageBusHeartbeatTimeout() const { return MessageBusHeartbeatTimeout; }
 	double GetMessageBusTimeBeforeRemovingDeadSource() const { return MessageBusTimeBeforeRemovingInactiveSource; }
 };
-
-#if UE_ENABLE_INCLUDE_ORDER_DEPRECATED_IN_5_2
-#include "Engine/EngineTypes.h"
-#include "LiveLinkFrameInterpolationProcessor.h"
-#include "LiveLinkFramePreProcessor.h"
-#include "LiveLinkRole.h"
-#include "LiveLinkSourceSettings.h"
-#include "LiveLinkSubjectSettings.h"
-#endif
diff --git a/Engine/Plugins/Animation/LiveLinkHub/Source/LiveLinkHub/Private/LiveLinkHubSubjectSettings.cpp b/Engine/Plugins/Animation/LiveLinkHub/Source/LiveLinkHub/Private/LiveLinkHubSubjectSettings.cpp
new file mode 100644
--- /dev/null
+++ b/Engine/Plugins/Animation/LiveLinkHub/Source/LiveLinkHub/Private/LiveLinkHubSubjectSettings.cpp
@@ -0,0 +1,97 @@
+// Copyright Epic Games, Inc. All Rights Reserved.
+
+#include "LiveLinkHubSubjectSettings.h"
+
+#include "LiveLinkHubClient.h"
+#include "LiveLinkHubModule.h"
+#include "Features/IModularFeatures.h"
+#include "Clients/LiveLinkHubProvider.h"
+
+void ULiveLinkHubSubjectSettings::Initialize(FLiveLinkSubjectKey InSubjectKey, TSubclassOf<ULiveLinkRole> InRole, FLiveLinkHubClient* InLiveLinkClient)
+{
+	SubjectName = InSubjectKey.SubjectName.ToString();
+	SubjectKey = InSubjectKey;
+
+	OutboundName = SubjectName;
+
+	Source = InLiveLinkClient->GetSourceType(InSubjectKey.Source).ToString();
+}
+
+void ULiveLinkHubSubjectSettings::NotifyRename()
+{
+	FLiveLinkHubModule& LiveLinkHubModule = FModuleManager::Get().GetModuleChecked<FLiveLinkHubModule>("LiveLinkHub");
+
+	if (TSharedPtr<FLiveLinkHubProvider> Provider = LiveLinkHubModule.GetLiveLinkProvider())
+	{
+		// Re-send the last static data with the new name.
+		TPair<UClass*, FLiveLinkStaticDataStruct*> StaticData = Provider->GetLastSubjectStaticDataStruct(PreviousOutboundName);
+		if (StaticData.Key && StaticData.Value)
+		{
+			FLiveLinkStaticDataStruct StaticDataCopy;
+			StaticDataCopy.InitializeWith(*StaticData.Value);
+
+			Provider->UpdateSubjectStaticData(*OutboundName, StaticData.Key, MoveTemp(StaticDataCopy));
+		}
+
+		// Then clear the old static data entry in the provider.
+		Provider->RemoveSubject(PreviousOutboundName);
+	}
+}
+
+void ULiveLinkHubSubjectSettings::PreEditChange(FProperty* PropertyAboutToChange)
+{
+	Super::PreEditChange(PropertyAboutToChange);
+
+	if (PropertyAboutToChange && PropertyAboutToChange->GetFName() == GET_MEMBER_NAME_CHECKED(ULiveLinkHubSubjectSettings, OutboundName))
+	{
+		PreviousOutboundName = *OutboundName;
+
+		FLiveLinkHubModule& LiveLinkHubModule = FModuleManager::Get().GetModuleChecked<FLiveLinkHubModule>("LiveLinkHub");
+		if (TSharedPtr<FLiveLinkHubProvider> Provider = LiveLinkHubModule.GetLiveLinkProvider())
+		{
+			Provider->SendClearSubjectToConnections(PreviousOutboundName);
+		}
+	}
+}
+
+void ULiveLinkHubSubjectSettings::PostEditChangeProperty(FPropertyChangedEvent& PropertyChangedEvent)
+{
+	Super::PostEditChangeProperty(PropertyChangedEvent);
+
+	if (PropertyChangedEvent.GetPropertyName() == GET_MEMBER_NAME_CHECKED(ULiveLinkHubSubjectSettings, OutboundName))
+	{
+		if (PreviousOutboundName != *OutboundName)
+		{
+			NotifyRename();
+		}
+	}
+	else if (PropertyChangedEvent.GetPropertyName() == GET_MEMBER_NAME_CHECKED(ULiveLinkHubSubjectSettings, Translators)
+		|| PropertyChangedEvent.GetPropertyName() == GET_MEMBER_NAME_CHECKED(ULiveLinkHubSubjectSettings, PreProcessors))
+	{
+		FLiveLinkHubClient* LiveLinkClient = static_cast<FLiveLinkHubClient*>(&IModularFeatures::Get().GetModularFeature<ILiveLinkClient>(ILiveLinkClient::ModularFeatureName));
+		LiveLinkClient->CacheSubjectSettings(SubjectKey, this);
+	}
+	else if (PropertyChangedEvent.GetPropertyName() == GET_MEMBER_NAME_CHECKED(ULiveLinkHubSubjectSettings, TranslatorsProxy))
+	{
+		Translators.Reset(1);
+		if (TranslatorsProxy)
+		{
+			Translators.Add(TranslatorsProxy);
+		}
+
+		ValidateProcessors();
+
+		// Re-assign TranslatorsProxy in case the translator was denied in the validate function.
+		if (Translators.Num())
+		{
+			TranslatorsProxy = Translators[0];
+		}
+		else
+		{
+			TranslatorsProxy = nullptr;
+		}
+
+		FLiveLinkHubClient* LiveLinkClient = static_cast<FLiveLinkHubClient*>(&IModularFeatures::Get().GetModularFeature<ILiveLinkClient>(ILiveLinkClient::ModularFeatureName));
+		LiveLinkClient->CacheSubjectSettings(SubjectKey, this);
+	}
+}
diff --git a/Engine/Plugins/Animation/LiveLinkHub/Source/LiveLinkHub/Private/LiveLinkHubSubjectSettings.h b/Engine/Plugins/Animation/LiveLinkHub/Source/LiveLinkHub/Private/LiveLinkHubSubjectSettings.h
new file mode 100644
--- /dev/null
+++ b/Engine/Plugins/Animation/LiveLinkHub/Source/LiveLinkHub/Private/LiveLinkHubSubjectSettings.h
@@ -0,0 +1,60 @@
+// Copyright Epic Games, Inc. All Rights Reserved.
+
+#pragma once
+
+#include "LiveLinkSubjectSettings.h"
+
+#include "LiveLinkHubSubjectSettings.generated.h"
+
+class FLiveLinkHubClient;
+class ULiveLinkRole;
+
+/** Settings object for a livelinkhub subject. */
+UCLASS()
+class ULiveLinkHubSubjectSettings : public ULiveLinkSubjectSettings
+{
+	GENERATED_BODY()
+public:
+	/** Fetch information from the livelink client to initialize this settings object. */
+	void Initialize(FLiveLinkSubjectKey InSubjectKey, TSubclassOf<ULiveLinkRole> InRole, FLiveLinkHubClient* InLiveLinkClient);
+
+	/** Notify connected clients that this subject's name has changed. */
+	void NotifyRename();
+
+	//~ Begin UObject interface
+	virtual void PreEditChange(FProperty* PropertyAboutToChange) override;
+	virtual void PostEditChangeProperty(FPropertyChangedEvent& PropertyChangedEvent) override;
+	//~ End UObject interface
+
+public:
+	/** Name of this subject. */
+	UPROPERTY(VisibleAnywhere, Category = "LiveLink")
+	FString SubjectName;
+
+	/** Name override that will be transmitted to clients instead of the subject name. */
+	UPROPERTY(EditAnywhere, Category = "LiveLink")
+	FString OutboundName;
+
+	/** Source that contains the subject. */
+	UPROPERTY(VisibleAnywhere, Category = "LiveLink")
+	FString Source;
+
+	/** Proxy property used edit the translators. */
+	UPROPERTY(EditAnywhere, Instanced, Category = "LiveLink", DisplayName = "Translator")
+	TObjectPtr<ULiveLinkFrameTranslator> TranslatorsProxy;
+
+private:
+	/* Previous outbound name to be used for noticing clients to remove this entry from their subject list. */
+	FName PreviousOutboundName;
+
+	/** SubjectKey of this subject. */
+	FLiveLinkSubjectKey SubjectKey;
+
+	struct FNameChangeInfo
+	{
+		UClass* SubjectRole = nullptr;
+		FLiveLinkStaticDataStruct LastStaticData;
+	};
+
+	friend class FLiveLinkHubSubjectSettingsDetailCustomization;
+};
diff --git a/Engine/Programs/LiveLinkHub/Config/DefaultGame.ini b/Engine/Programs/LiveLinkHub/Config/DefaultGame.ini
new file mode 100644
--- /dev/null
+++ b/Engine/Programs/LiveLinkHub/Config/DefaultGame.ini
@@ -0,0 +1,23 @@
+[/Script/LiveLink.LiveLinkSettings]
+-DefaultRoleSettings=(Role=Class'/Script/LiveLink.LiveLinkAnimationRole', FrameInterpolationProcessor=Class'/Script/LiveLink.LiveLinkAnimationFrameInterpolationProcessor')
++DefaultRoleSettings=(Role="/Script/CoreUObject.Class'/Script/LiveLinkInterface.LiveLinkAnimationRole'",SettingClass="/Script/CoreUObject.Class'/Script/LiveLinkHub.LiveLinkHubSubjectSettings'",FrameInterpolationProcessor="/Script/CoreUObject.Class'/Script/LiveLink.LiveLinkAnimationFrameInterpolationProcessor'",FramePreProcessors=)
+DefaultSettingsClass=/Script/LiveLinkHub.LiveLinkHubSubjectSettings
+FrameInterpolationProcessor=/Script/CoreUObject.Class'/Script/LiveLink.LiveLinkBasicFrameInterpolationProcessor'
+DefaultLiveLinkPreset=/Game/NewPreset.NewPreset
+ClockOffsetCorrectionStep=0.000100
+DefaultMessageBusSourceMode=EngineTime
+MessageBusPingRequestFrequency=1.000000
+MessageBusHeartbeatFrequency=1.000000
+MessageBusHeartbeatTimeout=2.000000
+MessageBusTimeBeforeRemovingInactiveSource=30.000000
+TimeWithoutFrameToBeConsiderAsInvalid=0.500000
+ValidColor=(R=0.000000,G=1.000000,B=0.000000,A=1.000000)
+InvalidColor=(R=1.000000,G=1.000000,B=0.000000,A=1.000000)
+TextSizeSource=16
+TextSizeSubject=12
+
+[/Script/LiveLinkHub.LiveLinkHubSettings]
+bAutoAddDiscoveredClients=True
++ProjectSettingsToShow=LiveLink
++ProjectSettingsToShow=LiveLinkHub
+
