LiveLinkHub -  Add LiveLinkhub Animation Virtual subject

Switch livelinkhub to use subject settings instead of subject session data.
This change will simplify the serialization of subject settings in a session config, and allow users to see frame / static data that's coming in the hub.

diff --git a/Engine/Plugins/Animation/LiveLink/Source/LiveLink/Private/LiveLinkClient.cpp b/Engine/Plugins/Animation/LiveLink/Source/LiveLink/Private/LiveLinkClient.cpp
--- a/Engine/Plugins/Animation/LiveLink/Source/LiveLink/Private/LiveLinkClient.cpp
+++ b/Engine/Plugins/Animation/LiveLink/Source/LiveLink/Private/LiveLinkClient.cpp
@@ -239,7 +239,7 @@ void FLiveLinkClient::HandleSubjectRebroadcast(ILiveLinkSubject* InSubject, cons
 			// Setup rebroadcast provider
 			if (!RebroadcastLiveLinkProvider.IsValid())
 			{
-				RebroadcastLiveLinkProvider = ILiveLinkProvider::CreateLiveLinkProvider(RebroadcastLiveLinkProviderName);
+				RebroadcastLiveLinkProvider = GetRebroadcastLiveLinkProvider();
 			}
 				
 			if (RebroadcastLiveLinkProvider.IsValid())
@@ -737,6 +737,11 @@ void FLiveLinkClient::BroadcastFrameDataUpdate(const FLiveLinkSubjectKey& InSubj
 	}
 }
 
+TSharedPtr<ILiveLinkProvider> FLiveLinkClient::GetRebroadcastLiveLinkProvider() const
+{
+	return ILiveLinkProvider::CreateLiveLinkProvider(RebroadcastLiveLinkProviderName);
+}
+
 void FLiveLinkClient::PushSubjectFrameData_AnyThread(const FLiveLinkSubjectKey& InSubjectKey, FLiveLinkFrameDataStruct&& InFrameData)
 {
 	FPendingSubjectFrame SubjectFrame{ InSubjectKey, MoveTemp(InFrameData) };
diff --git a/Engine/Plugins/Animation/LiveLink/Source/LiveLink/Public/LiveLinkClient.h b/Engine/Plugins/Animation/LiveLink/Source/LiveLink/Public/LiveLinkClient.h
--- a/Engine/Plugins/Animation/LiveLink/Source/LiveLink/Public/LiveLinkClient.h
+++ b/Engine/Plugins/Animation/LiveLink/Source/LiveLink/Public/LiveLinkClient.h
@@ -244,6 +244,9 @@ protected:
 	/** Broadcast out to the SubjectFrameAddedHandles a frame data update. */
 	void BroadcastFrameDataUpdate(const FLiveLinkSubjectKey& InSubjectKey, const FLiveLinkFrameDataStruct& InFrameData);
 
+	/** Method that can be overriden in child classes to provide their own rebroadcast provider. */
+	virtual TSharedPtr<ILiveLinkProvider> GetRebroadcastLiveLinkProvider() const;
+
 	/** The current collection used. */
 	TUniquePtr<FLiveLinkSourceCollection> Collection;
 
diff --git a/Engine/Plugins/Animation/LiveLink/Source/LiveLinkEditor/Private/LiveLinkClientPanelToolbar.cpp b/Engine/Plugins/Animation/LiveLink/Source/LiveLinkEditor/Private/LiveLinkClientPanelToolbar.cpp
--- a/Engine/Plugins/Animation/LiveLink/Source/LiveLinkEditor/Private/LiveLinkClientPanelToolbar.cpp
+++ b/Engine/Plugins/Animation/LiveLink/Source/LiveLinkEditor/Private/LiveLinkClientPanelToolbar.cpp
@@ -613,7 +613,6 @@ TSharedRef<SWidget> SLiveLinkClientPanelToolbar::OnGenerateSourceMenu()
 	//	FNewMenuDelegate::CreateRaw(this, &SLiveLinkClientPanelToolbar::PopulateVirtualSubjectSourceCreationMenu)
 	//);
 
-#if !WITH_LIVELINK_HUB
 	MenuBuilder.AddMenuEntry(
 		LOCTEXT("AddVirtualSubject", "Add Virtual Subject"),
 		LOCTEXT("AddVirtualSubject_Tooltip", "Adds a new virtual subject to LiveLink. Instead of coming from a source a virtual subject is a combination of 2 or more real subjects"),
@@ -623,7 +622,6 @@ TSharedRef<SWidget> SLiveLinkClientPanelToolbar::OnGenerateSourceMenu()
 		),
 		NAME_None,
 		EUserInterfaceActionType::Button);
-#endif
 
 	MenuBuilder.EndSection();
 
diff --git a/Engine/Plugins/Animation/LiveLink/Source/LiveLinkEditor/Public/LiveLinkClientPanelViews.h b/Engine/Plugins/Animation/LiveLink/Source/LiveLinkEditor/Public/LiveLinkClientPanelViews.h
--- a/Engine/Plugins/Animation/LiveLink/Source/LiveLinkEditor/Public/LiveLinkClientPanelViews.h
+++ b/Engine/Plugins/Animation/LiveLink/Source/LiveLinkEditor/Public/LiveLinkClientPanelViews.h
@@ -32,8 +32,8 @@ typedef TSharedPtr<FLiveLinkSubjectUIEntry> FLiveLinkSubjectUIEntryPtr;
 
 namespace UE::LiveLink
 {
-	TSharedPtr<IDetailsView> CreateSourcesDetailsView(const TSharedPtr<FLiveLinkSourcesView>& InSourcesView, const TAttribute<bool>& bInReadOnly);
-	TSharedPtr<SLiveLinkDataView> CreateSubjectsDetailsView(FLiveLinkClient* InLiveLinkClient, const TAttribute<bool>& bInReadOnly);
+	TSharedPtr<IDetailsView> LIVELINKEDITOR_API CreateSourcesDetailsView(const TSharedPtr<FLiveLinkSourcesView>& InSourcesView, const TAttribute<bool>& bInReadOnly);
+	TSharedPtr<SLiveLinkDataView> LIVELINKEDITOR_API CreateSubjectsDetailsView(FLiveLinkClient* InLiveLinkClient, const TAttribute<bool>& bInReadOnly);
 }
 
 // Structure that defines a single entry in the subject UI
diff --git a/Engine/Plugins/Animation/LiveLink/Source/LiveLinkEditor/Public/SLiveLinkDataView.h b/Engine/Plugins/Animation/LiveLink/Source/LiveLinkEditor/Public/SLiveLinkDataView.h
--- a/Engine/Plugins/Animation/LiveLink/Source/LiveLinkEditor/Public/SLiveLinkDataView.h
+++ b/Engine/Plugins/Animation/LiveLink/Source/LiveLinkEditor/Public/SLiveLinkDataView.h
@@ -13,7 +13,7 @@ class IDetailsView;
 class FLiveLinkClient;
 class IStructureDetailsView;
 
-class SLiveLinkDataView : public SCompoundWidget
+class LIVELINKEDITOR_API SLiveLinkDataView : public SCompoundWidget
 {
 private:
 	using Super = SCompoundWidget;
diff --git a/Engine/Plugins/Animation/LiveLinkHub/Source/LiveLinkHub/Private/Clients/LiveLinkHubProvider.cpp b/Engine/Plugins/Animation/LiveLinkHub/Source/LiveLinkHub/Private/Clients/LiveLinkHubProvider.cpp
--- a/Engine/Plugins/Animation/LiveLinkHub/Source/LiveLinkHub/Private/Clients/LiveLinkHubProvider.cpp
+++ b/Engine/Plugins/Animation/LiveLinkHub/Source/LiveLinkHub/Private/Clients/LiveLinkHubProvider.cpp
@@ -20,7 +20,6 @@
 #include "Misc/ScopeLock.h"
 #include "Session/LiveLinkHubSession.h"
 #include "Session/LiveLinkHubSessionManager.h"
-#include "Subjects/LiveLinkHubSubjectSessionConfig.h"
 #include "TimerManager.h"
 
 
diff --git a/Engine/Plugins/Animation/LiveLinkHub/Source/LiveLinkHub/Private/LiveLinkHub.cpp b/Engine/Plugins/Animation/LiveLinkHub/Source/LiveLinkHub/Private/LiveLinkHub.cpp
--- a/Engine/Plugins/Animation/LiveLinkHub/Source/LiveLinkHub/Private/LiveLinkHub.cpp
+++ b/Engine/Plugins/Animation/LiveLinkHub/Source/LiveLinkHub/Private/LiveLinkHub.cpp
@@ -2,26 +2,21 @@
 
 #include "LiveLinkHub.h"
 
+
 #include "AssetRegistry/AssetRegistryModule.h"
 #include "AssetRegistry/IAssetRegistry.h"
 #include "Clients/LiveLinkHubClientsController.h"
 #include "Clients/LiveLinkHubProvider.h"
-#include "Config/LiveLinkHubFileUtilities.h"
-#include "DesktopPlatformModule.h"
-#include "EditorDirectories.h"
 #include "Features/IModularFeatures.h"
 #include "Framework/Application/SlateApplication.h"
-#include "IDesktopPlatform.h"
 #include "ISettingsModule.h"
 #include "LiveLinkEditorSettings.h"
 #include "LiveLinkHubClient.h"
 #include "LiveLinkHubSettings.h"
-#include "LiveLinkProvider.h"
+#include "LiveLinkHubSubjectSettings.h"
 #include "LiveLinkHubCommands.h"
 #include "LiveLinkProviderImpl.h"
 #include "LiveLinkSettings.h"
-#include "LiveLinkSubject.h"
-#include "LiveLinkSubjectSettings.h"
 #include "Misc/App.h"
 #include "Misc/ConfigCacheIni.h"
 #include "Recording/LiveLinkHubPlaybackController.h"
@@ -230,15 +225,10 @@ void FLiveLinkHub::OpenConfig()
 
 FName FLiveLinkHub::GetSubjectNameOverride(const FLiveLinkSubjectKey& InSubjectKey) const
 {
-	if (const TSharedPtr<ILiveLinkHubSessionManager> Manager = SessionManager)
+	ILiveLinkClient& Client = IModularFeatures::Get().GetModularFeature<ILiveLinkClient>(ILiveLinkClient::ModularFeatureName);
+	if (ULiveLinkHubSubjectSettings* Settings = Cast<ULiveLinkHubSubjectSettings>(Client.GetSubjectSettings(InSubjectKey)))
 	{
-		if (const TSharedPtr<ILiveLinkHubSession> CurrentSession = Manager->GetCurrentSession())
-		{
-			if (TOptional<FLiveLinkHubSubjectProxy> SubjectProxy = CurrentSession->GetSubjectConfig(InSubjectKey))
-			{
-				return SubjectProxy->GetOutboundName();
-			}
-		}
+		return *Settings->OutboundName;
 	}
 
 	return InSubjectKey.SubjectName;
diff --git a/Engine/Plugins/Animation/LiveLinkHub/Source/LiveLinkHub/Private/LiveLinkHubAnimationVirtualSubject.h b/Engine/Plugins/Animation/LiveLinkHub/Source/LiveLinkHub/Private/LiveLinkHubAnimationVirtualSubject.h
new file mode 100644
--- /dev/null
+++ b/Engine/Plugins/Animation/LiveLinkHub/Source/LiveLinkHub/Private/LiveLinkHubAnimationVirtualSubject.h
@@ -0,0 +1,65 @@
+// Copyright Epic Games, Inc. All Rights Reserved.
+
+#pragma once
+
+#include "LiveLinkAnimationVirtualSubject.h"
+
+#include "Clients/LiveLinkHubProvider.h"
+#include "Features/IModularFeatures.h"
+#include "LiveLinkClient.h"
+#include "LiveLinkHubModule.h"
+
+#include "LiveLinkHubAnimationVirtualSubject.generated.h"
+
+/**
+ * Animation virtual subject used in LiveLinkHub.
+ * Shows options for the subject and broadcasts static data when the skeleton is updated.
+ */
+UCLASS()
+class ULiveLinkHubAnimationVirtualSubject : public ULiveLinkAnimationVirtualSubject
+{
+	GENERATED_BODY()
+
+public:
+	virtual void Initialize(FLiveLinkSubjectKey InSubjectKey, TSubclassOf<ULiveLinkRole> InRole, ILiveLinkClient* InLiveLinkClient) override
+	{
+		ULiveLinkAnimationVirtualSubject::Initialize(InSubjectKey, InRole, InLiveLinkClient);
+		SubjectName = InSubjectKey.SubjectName.ToString();
+
+		FLiveLinkClient* Client = &IModularFeatures::Get().GetModularFeature<FLiveLinkClient>(FLiveLinkClient::ModularFeatureName);
+		Source = Client->GetSourceType(InSubjectKey.Source).ToString();
+	}
+
+	/** Whether this subject is rebroadcasted */
+	virtual bool IsRebroadcasted() const
+	{
+		// todo: Decide this based on the session? Provider should decide how to handle it .
+		return true;
+	}
+
+	virtual void PostSkeletonRebuild() override
+	{
+		if (HasValidStaticData())
+		{
+			if (const TSharedPtr<FLiveLinkHubProvider> LiveLinkProvider = FModuleManager::Get().GetModuleChecked<FLiveLinkHubModule>("LiveLinkHub").GetLiveLinkProvider())
+			{
+				// Update the static data since the final skeleton was changed.
+				const FLiveLinkSubjectFrameData& CurrentSnapshot = GetFrameSnapshot();
+
+				FLiveLinkStaticDataStruct StaticDataCopy;
+				StaticDataCopy.InitializeWith(CurrentSnapshot.StaticData);
+
+				LiveLinkProvider->UpdateSubjectStaticData(SubjectKey.SubjectName, Role, MoveTemp(StaticDataCopy));
+			}
+		}
+	}
+
+public:
+	/** Name of the virtual subject. */
+	UPROPERTY(VisibleAnywhere, Category = "LiveLink")
+	FString SubjectName;
+
+	/** Source that contains the subject. */
+	UPROPERTY(VisibleAnywhere, Category = "LiveLink")
+	FString Source;
+};
diff --git a/Engine/Plugins/Animation/LiveLinkHub/Source/LiveLinkHub/Private/LiveLinkHubClient.cpp b/Engine/Plugins/Animation/LiveLinkHub/Source/LiveLinkHub/Private/LiveLinkHubClient.cpp
--- a/Engine/Plugins/Animation/LiveLinkHub/Source/LiveLinkHub/Private/LiveLinkHubClient.cpp
+++ b/Engine/Plugins/Animation/LiveLinkHub/Source/LiveLinkHub/Private/LiveLinkHubClient.cpp
@@ -2,11 +2,15 @@
 
 #include "LiveLinkHubClient.h"
 
+#include "Algo/AnyOf.h"
 #include "Async/TaskGraphInterfaces.h"
+#include "Clients/LiveLinkHubProvider.h"
 #include "Engine/Engine.h"
 #include "LiveLinkClient.h"
 #include "LiveLinkHub.h"
 #include "LiveLinkHubLog.h"
+#include "LiveLinkHubModule.h"
+#include "LiveLinkHubSubjectSettings.h"
 #include "LiveLinkRoleTrait.h"
 #include "LiveLinkSourceCollection.h"
 #include "LiveLinkSettings.h"
@@ -28,8 +32,7 @@ DECLARE_CYCLE_STAT(TEXT("LiveLinkHub - Push FrameData"), STAT_LiveLinkHub_PushFr
 #define LOCTEXT_NAMESPACE "LiveLinkHub.LiveLinkHubClient"
 
 FLiveLinkHubClient::FLiveLinkHubClient(TSharedPtr<ILiveLinkHub> InLiveLinkHub)
-	: FLiveLinkClient()
-	, LiveLinkHub(MoveTemp(InLiveLinkHub))
+	: LiveLinkHub(MoveTemp(InLiveLinkHub))
 {
 }
 
@@ -52,20 +55,6 @@ void FLiveLinkHubClient::CacheSubjectSettings(const FLiveLinkSubjectKey& Subject
 	}
 }
 
-const FLiveLinkStaticDataStruct* FLiveLinkHubClient::GetSubjectStaticData(const FLiveLinkSubjectKey& InSubjectKey)
-{
-	FScopeLock Lock(&CollectionAccessCriticalSection);
-	if (const FLiveLinkCollectionSubjectItem* SubjectItem = Collection->FindSubject(InSubjectKey))
-	{
-		if (FLiveLinkSubject* LiveLinkSubject = SubjectItem->GetLiveSubject())
-		{
-			return &LiveLinkSubject->GetStaticData();
-		}
-	}
-
-	return nullptr;
-}
-
 bool FLiveLinkHubClient::CreateSource(const FLiveLinkSourcePreset& InSourcePreset)
 {
 	// Create fake sources if we're in playback mode so we don't process livelink data while we're playing a recording.
@@ -327,10 +316,16 @@ void FLiveLinkHubClient::PushSubjectStaticData_AnyThread(const FLiveLinkSubjectK
 			UClass* SettingClass = DefaultSetting.SettingClass.Get();
 			if (SettingClass == nullptr)
 			{
-				SettingClass = ULiveLinkSubjectSettings::StaticClass();
+				SettingClass = ULiveLinkHubSubjectSettings::StaticClass();
 			}
 
 			SubjectSettings = NewObject<ULiveLinkSubjectSettings>(GetTransientPackage(), SettingClass);
+
+			if (SubjectSettings->IsA<ULiveLinkHubSubjectSettings>())
+			{
+				Cast<ULiveLinkHubSubjectSettings>(SubjectSettings)->Initialize(SubjectKey, Role.Get(), this);
+			}
+
 			SubjectSettings->Role = Role;
 
 			UClass* FrameInterpolationProcessorClass = DefaultSetting.FrameInterpolationProcessor.Get();
@@ -383,6 +378,7 @@ void FLiveLinkHubClient::PushSubjectStaticData_AnyThread(const FLiveLinkSubjectK
 
 			bool bEnabled = Collection->FindEnabledSubject(SubjectKey.SubjectName) == nullptr;
 			FLiveLinkCollectionSubjectItem CollectionSubjectItem(SubjectKey, MakeUnique<FLiveLinkSubject>(SourceItem->TimedData), SubjectSettings, bEnabled);
+			
 			CollectionSubjectItem.GetLiveSubject()->Initialize(SubjectKey, Role.Get(), this);
 
 			LiveLinkSubject = CollectionSubjectItem.GetLiveSubject();
@@ -401,26 +397,27 @@ void FLiveLinkHubClient::PushSubjectStaticData_AnyThread(const FLiveLinkSubjectK
 		// Dispatch it on the game thread since FLiveLinkSubject::SetStaticData asserts when called outside the game thread.
 		// Note that SetStaticData may not be needed on the hub but we might need it for interpolation or preprocessing in the future.
 		FFunctionGraphTask::CreateAndDispatchWhenReady([this, LiveLinkSubject, SubjectKey, Role, StaticData = MoveTemp(InStaticData)]() mutable
+		{
+			FScopeLock Lock(&CollectionAccessCriticalSection);
+			// Validate live link subject hasn't changed
+			if (FLiveLinkCollectionSubjectItem* SubjectItem = Collection->FindSubject(SubjectKey))
 			{
-				FScopeLock Lock(&CollectionAccessCriticalSection);
-				// Validate live link subject hasn't changed
-				if (FLiveLinkCollectionSubjectItem* SubjectItem = Collection->FindSubject(SubjectKey))
+				const FLiveLinkSubject* CurrentLiveLinkSubject = SubjectItem->GetLiveSubject();
+				if (LiveLinkSubject == CurrentLiveLinkSubject)
 				{
-					const FLiveLinkSubject* CurrentLiveLinkSubject = SubjectItem->GetLiveSubject();
-					if (LiveLinkSubject == CurrentLiveLinkSubject)
-					{
-						LiveLinkSubject->SetStaticData(Role, MoveTemp(StaticData));
-					}
+    				LiveLinkSubject->SetStaticData(Role, MoveTemp(StaticData));
 				}
-			}, TStatId(), nullptr, ENamedThreads::GameThread);
+			}
+		}, TStatId(), nullptr, ENamedThreads::GameThread);
 	}
-
 }
 
 void FLiveLinkHubClient::PushSubjectFrameData_AnyThread(const FLiveLinkSubjectKey& SubjectKey, FLiveLinkFrameDataStruct&& FrameData)
 {
 	SCOPE_CYCLE_COUNTER(STAT_LiveLinkHub_PushFrameData);
 
+	TOptional<FLiveLinkSubjectFrameData> TranslatedFrame;
+
 	bool bFrameValid = true;
 	{
 		FScopeLock Lock(&CollectionAccessCriticalSection);
@@ -438,21 +435,34 @@ void FLiveLinkHubClient::PushSubjectFrameData_AnyThread(const FLiveLinkSubjectKe
 					TArray<ULiveLinkFrameTranslator::FWorkerSharedPtr> Translators = LiveSubject->GetFrameTranslators();
 					if (Translators.Num() && Translators[0].IsValid())
 					{
-						FLiveLinkSubjectFrameData OutTranslatedFrame;
-						Translators[0]->Translate(LiveSubject->GetStaticData(), FrameData, OutTranslatedFrame);
-						FrameData = MoveTemp(OutTranslatedFrame.FrameData);
+						FLiveLinkSubjectFrameData TranslatedFrameData;
+						if (Translators[0]->Translate(LiveSubject->GetStaticData(), FrameData, TranslatedFrameData))
+						{
+							TranslatedFrame = MoveTemp(TranslatedFrameData);
+						}
 					}
 				}
 			}
 		}
 	}
 
+	// If the data was translated, then we broadcast this data, but we save snapshots of the original data since Evaluate will need it for translation as well.
+	FLiveLinkFrameDataStruct* FrameDataToBroadcast = TranslatedFrame.IsSet() ? &(TranslatedFrame->FrameData) : &FrameData;
+
 	if (bFrameValid)
 	{
-		OnFrameDataReceivedDelegate_AnyThread.Broadcast(SubjectKey, FrameData);
+		OnFrameDataReceivedDelegate_AnyThread.Broadcast(SubjectKey, *FrameDataToBroadcast);
 	}
 
-	BroadcastFrameDataUpdate(SubjectKey, FrameData);
+	if (bVirtualSubjectsPresent)
+	{
+		// Cache frames if we have virtual subjects so they can use the data.
+		FLiveLinkClient::PushSubjectFrameData_AnyThread(SubjectKey, MoveTemp(FrameData));
+	}
+	else
+	{
+		BroadcastFrameDataUpdate(SubjectKey, MoveTemp(*FrameDataToBroadcast));
+	}
 }
 
 FText FLiveLinkHubClient::GetSourceStatus(FGuid InEntryGuid) const
@@ -483,9 +493,36 @@ bool FLiveLinkHubClient::IsSubjectValid(const FLiveLinkSubjectKey& InSubjectKey)
 void FLiveLinkHubClient::RemoveSubject_AnyThread(const FLiveLinkSubjectKey& InSubjectKey)
 {
 	OnSubjectMarkedPendingKill_AnyThread().Broadcast(InSubjectKey);
+
 	FLiveLinkClient::RemoveSubject_AnyThread(InSubjectKey);
 }
 
+bool FLiveLinkHubClient::AddVirtualSubject(const FLiveLinkSubjectKey& VirtualSubjectKey, TSubclassOf<ULiveLinkVirtualSubject> VirtualSubjectClass)
+{
+	bool bResult =  FLiveLinkClient::AddVirtualSubject(VirtualSubjectKey, VirtualSubjectClass);
+	if (bResult)
+	{
+		bVirtualSubjectsPresent = true;
+	}
+
+	return bResult;
+}
+
+void FLiveLinkHubClient::RemoveVirtualSubject(const FLiveLinkSubjectKey& VirtualSubjectKey)
+{
+	FLiveLinkClient::RemoveVirtualSubject(VirtualSubjectKey);
+
+	{
+		FScopeLock Lock(&CollectionAccessCriticalSection);
+		bVirtualSubjectsPresent = Algo::AnyOf(Collection->GetSubjects(), [] (const FLiveLinkCollectionSubjectItem& SubjectItem) { return !!SubjectItem.GetVirtualSubject(); });
+	}
+}
+
+TSharedPtr<ILiveLinkProvider> FLiveLinkHubClient::GetRebroadcastLiveLinkProvider() const
+{
+	return FModuleManager::Get().GetModuleChecked<FLiveLinkHubModule>("LiveLinkHub").GetLiveLinkProvider();
+}
+
 void FLiveLinkHubClient::BroadcastStaticDataUpdate(FLiveLinkSubject* InLiveSubject, TSubclassOf<ULiveLinkRole> InRole, const FLiveLinkStaticDataStruct& InStaticData) const
 {
 	// If we have a translator, we need to broadcast the translated static data to the clients.
diff --git a/Engine/Plugins/Animation/LiveLinkHub/Source/LiveLinkHub/Private/LiveLinkHubClient.h b/Engine/Plugins/Animation/LiveLinkHub/Source/LiveLinkHub/Private/LiveLinkHubClient.h
--- a/Engine/Plugins/Animation/LiveLinkHub/Source/LiveLinkHub/Private/LiveLinkHubClient.h
+++ b/Engine/Plugins/Animation/LiveLinkHub/Source/LiveLinkHub/Private/LiveLinkHubClient.h
@@ -24,9 +24,6 @@ class FLiveLinkHubClient : public FLiveLinkClient
 public:
 	FLiveLinkHubClient(TSharedPtr<ILiveLinkHub> InLiveLinkHub);
 	virtual ~FLiveLinkHubClient();
-
-	/** Utility method to grab a subject's static data. Used by the RecordingController when static data is missing from the recording. */
-	const FLiveLinkStaticDataStruct* GetSubjectStaticData(const FLiveLinkSubjectKey& InSubjectKey);
 	
 	/** Get the delegate called when frame data is received. */
 	FOnFrameDataReceived_AnyThread& OnFrameDataReceived_AnyThread()
@@ -61,8 +58,14 @@ public:
 	virtual FText GetSourceStatus(FGuid InEntryGuid) const override;
 	virtual bool IsSubjectValid(const FLiveLinkSubjectKey& InSubjectKey) const override;
 	virtual void RemoveSubject_AnyThread(const FLiveLinkSubjectKey& InSubjectKey) override;
+	virtual bool AddVirtualSubject(const FLiveLinkSubjectKey& VirtualSubjectKey, TSubclassOf<ULiveLinkVirtualSubject> VirtualSubjectClass) override;
+    virtual void RemoveVirtualSubject(const FLiveLinkSubjectKey& VirtualSubjectKey) override;
 	//~ End ILiveLinkClient interface
 
+	//~ Begin FLiveLinkClient interface
+	virtual TSharedPtr<ILiveLinkProvider> GetRebroadcastLiveLinkProvider() const override;
+	//~ End FLiveLinkClient interface
+
 private:
 	/** Create a LiveLinkPlaybackSource which acts as a dummy source when doing playback. */
 	bool CreatePlaybackSource(const FLiveLinkSourcePreset& InSourcePreset);
@@ -70,8 +73,6 @@ private:
 	bool CreatePlaybackSubject(const FLiveLinkSubjectPreset& InSubjectPreset);
 	/** Broadcast a static data update to this client's listeners. */
 	void BroadcastStaticDataUpdate(FLiveLinkSubject* InLiveSubject, TSubclassOf<ULiveLinkRole> InRole, const FLiveLinkStaticDataStruct& InStaticData) const;
-	/** Lock to stop multiple threads accessing the Subjects from the collection at the same time */
-	mutable FCriticalSection CollectionAccessCriticalSection;
 
 private:
 	/** Weak pointer to the live link hub. */
@@ -82,4 +83,6 @@ private:
     FOnStaticDataReceived_AnyThread OnStaticDataReceivedDelegate_AnyThread;
 	/** Delegate called when a subject is marked for deletion. */
 	FOnSubjectMarkedPendingKill_AnyThread OnSubjectMarkedPendingKillDelegate_AnyThread;
+	/** Whether there are virtual subjects at the moment. Used to determine if we should cache frame data for their usage. */
+	std::atomic<bool> bVirtualSubjectsPresent = false;
 };
diff --git a/Engine/Plugins/Animation/LiveLinkHub/Source/LiveLinkHub/Private/Recording/Implementations/LiveLinkUAssetRecorder.cpp b/Engine/Plugins/Animation/LiveLinkHub/Source/LiveLinkHub/Private/Recording/Implementations/LiveLinkUAssetRecorder.cpp
--- a/Engine/Plugins/Animation/LiveLinkHub/Source/LiveLinkHub/Private/Recording/Implementations/LiveLinkUAssetRecorder.cpp
+++ b/Engine/Plugins/Animation/LiveLinkHub/Source/LiveLinkHub/Private/Recording/Implementations/LiveLinkUAssetRecorder.cpp
@@ -39,7 +39,7 @@ namespace UAssetRecorderUtils
 		FLiveLinkHubClient* LiveLinkClient = static_cast<FLiveLinkHubClient*>(&IModularFeatures::Get().GetModularFeature<ILiveLinkClient>(ILiveLinkClient::ModularFeatureName));
 		TSubclassOf<ULiveLinkRole> LiveLinkRole = LiveLinkClient->GetSubjectRole_AnyThread(SubjectKey);
 		
-		if (const FLiveLinkStaticDataStruct* StaticData = LiveLinkClient->GetSubjectStaticData(SubjectKey))
+		if (const FLiveLinkStaticDataStruct* StaticData = LiveLinkClient->GetSubjectStaticData_AnyThread(SubjectKey))
 		{
 			FInstancedStruct StaticDataInstancedStruct;
 			StaticDataInstancedStruct.InitializeAs(StaticData->GetStruct(), (uint8*)StaticData->GetBaseData());
diff --git a/Engine/Plugins/Animation/LiveLinkHub/Source/LiveLinkHub/Private/Session/LiveLinkHubSession.h b/Engine/Plugins/Animation/LiveLinkHub/Source/LiveLinkHub/Private/Session/LiveLinkHubSession.h
--- a/Engine/Plugins/Animation/LiveLinkHub/Source/LiveLinkHub/Private/Session/LiveLinkHubSession.h
+++ b/Engine/Plugins/Animation/LiveLinkHub/Source/LiveLinkHub/Private/Session/LiveLinkHubSession.h
@@ -11,9 +11,9 @@
 #include "LiveLinkHubModule.h"
 #include "LiveLinkHubSessionData.h"
 #include "LiveLinkTypes.h"
+#include "LiveLinkVirtualSubject.h"
 #include "Misc/ScopeLock.h"
 #include "Modules/ModuleManager.h"
-#include "Subjects/LiveLinkHubSubjectSessionConfig.h"
 #include "UObject/Package.h"
 #include "UObject/StrongObjectPtr.h"
 #include "UObject/UObjectGlobals.h"
@@ -24,9 +24,6 @@
 DECLARE_MULTICAST_DELEGATE_OneParam(FOnClientAddedToSession, FLiveLinkHubClientId);
 DECLARE_MULTICAST_DELEGATE_OneParam(FOnClientRemovedFromSession, FLiveLinkHubClientId);
 
-class ULiveLinkFramePreProcessor;
-class ULiveLinkFrameTranslator;
-
 /**
  * Holds the state of the hub for an active session, can be swapped out with a different session using the session manager.
  */
@@ -35,15 +32,6 @@ class ILiveLinkHubSession
 public:
 	virtual ~ILiveLinkHubSession() = default;
 
-	/** Get the configuration for a given subject. */
-	virtual TOptional<FLiveLinkHubSubjectProxy> GetSubjectConfig(const FLiveLinkSubjectKey& SubjectKey) const = 0;
-
-	/** Get the preprocessors and translator for a given subject. */
-	virtual ULiveLinkHubSubjectProcessors* GetSubjectProcessors(const FLiveLinkSubjectKey& SubjectKey) const = 0;
-
-	/** Change the outbound name of a subject for this session. */
-	virtual void RenameSubject(const FLiveLinkSubjectKey& SubjectKey, FName NewName) = 0;
-
 	/** Add a client to this session. Note: Must be called from game thread. */
 	virtual void AddClient(const FLiveLinkHubClientId& Client) = 0;
 
@@ -53,12 +41,6 @@ public:
 	/** Returns whether a client is in this session. */
 	virtual bool IsClientInSession(const FLiveLinkHubClientId& Client) = 0;
 
-	/** Set the preprocessors for a subject. */
-	virtual void SetPreProcessors(const FLiveLinkSubjectKey& SubjectKey, TConstArrayView<ULiveLinkFramePreProcessor*> PreProcessors) = 0;
-
-	/** Set the translators for a subject. */
-	virtual void SetTranslator(const FLiveLinkSubjectKey& SubjectKey, ULiveLinkFrameTranslator* Translator) = 0;
-
 	/** Get the list of clients in this session (The list of clients that can receive data from the hub) */
 	virtual TArray<FLiveLinkHubClientId> GetSessionClients() const = 0;
 };
@@ -70,10 +52,7 @@ public:
 		: OnClientAddedToSessionDelegate(OnClientAddedToSession)
 		, OnClientRemovedFromSessionDelegate(OnClientRemovedFromSession)
 	{
-		RegisterDelegates();
-
 		SessionData = TStrongObjectPtr<ULiveLinkHubSessionData>(NewObject<ULiveLinkHubSessionData>(GetTransientPackage()));
-		SessionData->SubjectsConfig->Initialize();
 	}
 
 	FLiveLinkHubSession(ULiveLinkHubSessionData* InSessionData, FOnClientAddedToSession& OnClientAddedToSession, FOnClientRemovedFromSession& OnClientRemovedFromSession)
@@ -81,46 +60,6 @@ public:
 		, OnClientRemovedFromSessionDelegate(OnClientRemovedFromSession)
 	{
 		SessionData = TStrongObjectPtr<ULiveLinkHubSessionData>(InSessionData);
-		RegisterDelegates();
-		SessionData->SubjectsConfig->Initialize();
-	}
-
-	virtual ~FLiveLinkHubSession() override
-	{
-		UnregisterDelegates();
-	}
-
-	virtual void RenameSubject(const FLiveLinkSubjectKey& SubjectKey, FName NewName) override
-	{
-		TMap<FLiveLinkSubjectKey, FLiveLinkHubSubjectProxy> SubjectProxies;
-		{
-			FReadScopeLock Locker(SessionDataLock);
-
-			SubjectProxies = SessionData->SubjectsConfig->SubjectProxies;
-		}
-
-		if (FLiveLinkHubSubjectProxy* Proxy = SubjectProxies.Find(SubjectKey))
-		{
-			Proxy->SetOutboundName(NewName);
-		}
-
-		{
-			// Copied over in a different step to avoid acquiring the rw lock in a method called by RenameSubject
-			FWriteScopeLock Locker(SessionDataLock);
-			SessionData->SubjectsConfig->SubjectProxies = MoveTemp(SubjectProxies);
-		}
-	}
-
-	virtual TOptional<FLiveLinkHubSubjectProxy> GetSubjectConfig(const FLiveLinkSubjectKey& SubjectKey) const override
-	{
-		FReadScopeLock Locker(SessionDataLock);
-		return SessionData->SubjectsConfig->GetSubjectConfig(SubjectKey);
-	}
-
-	virtual ULiveLinkHubSubjectProcessors* GetSubjectProcessors(const FLiveLinkSubjectKey& SubjectKey) const override
-	{
-		FReadScopeLock Locker(SessionDataLock);
-		return SessionData->SubjectsConfig->GetSubjectProcessors(SubjectKey);
 	}
 
 	virtual TArray<FLiveLinkHubClientId> GetSessionClients() const override
@@ -165,24 +104,6 @@ public:
 		return CachedSessionClients.Contains(Client);
 	}
 
-	virtual void SetPreProcessors(const FLiveLinkSubjectKey& SubjectKey, TConstArrayView<ULiveLinkFramePreProcessor*> PreProcessors) override
-	{
-		FWriteScopeLock Locker(SessionDataLock);
-		if (TObjectPtr<ULiveLinkHubSubjectProcessors>* SubjectProcessors = SessionData->SubjectsConfig->SubjectProcessors.Find(SubjectKey))
-		{
-			(*SubjectProcessors)->PreProcessors = PreProcessors;
-		}
-	}
-
-	virtual void SetTranslator(const FLiveLinkSubjectKey& SubjectKey, ULiveLinkFrameTranslator* Translator) override
-	{
-		FWriteScopeLock Locker(SessionDataLock);
-		if (TObjectPtr<ULiveLinkHubSubjectProcessors>* SubjectProcessors = SessionData->SubjectsConfig->SubjectProcessors.Find(SubjectKey))
-		{
-			(*SubjectProcessors)->Translator = Translator;
-		}
-	}
-
 	void AddRestoredClient(FLiveLinkHubUEClientInfo& InOutRestoredClientInfo)
 	{
 		if (const TSharedPtr<FLiveLinkHubProvider> LiveLinkProvider = FModuleManager::Get().GetModuleChecked<FLiveLinkHubModule>("LiveLinkHub").GetLiveLinkProvider())
@@ -198,82 +119,6 @@ public:
 		OnClientAddedToSessionDelegate.Broadcast(InOutRestoredClientInfo.Id);
 	}
 
-private:
-	/** Register livelink client delegates used to update the config's subject data. */
-	void RegisterDelegates()
-	{
-		FLiveLinkClient& LiveLinkClient = IModularFeatures::Get().GetModularFeature<FLiveLinkClient>(ILiveLinkClient::ModularFeatureName);
-		LiveLinkClient.OnLiveLinkSubjectAdded().AddRaw(this, &FLiveLinkHubSession::OnSubjectAdded_AnyThread);
-		LiveLinkClient.OnLiveLinkSubjectRemoved().AddRaw(this, &FLiveLinkHubSession::OnSubjectRemoved_AnyThread);
-	}
-
-	/** Unregister livelink client delegates. */
-	void UnregisterDelegates()
-	{
-		if (IModularFeatures::Get().IsModularFeatureAvailable(ILiveLinkClient::ModularFeatureName))
-		{
-			FLiveLinkClient& LiveLinkClient = IModularFeatures::Get().GetModularFeature<FLiveLinkClient>(ILiveLinkClient::ModularFeatureName);
-			LiveLinkClient.OnLiveLinkSubjectRemoved().RemoveAll(this);
-			LiveLinkClient.OnLiveLinkSubjectAdded().RemoveAll(this);
-		}
-	}
-
-	/** AnyThread handler for the SubjectAdded delegate, dispatches handling on the game thread to avoid asserts in Slate. */
-	void OnSubjectAdded_AnyThread(FLiveLinkSubjectKey SubjectKey)
-	{
-		TWeakPtr<FLiveLinkHubSession> WeakSession = AsShared();
-		AsyncTask(ENamedThreads::GameThread, [WeakSession, Key = MoveTemp(SubjectKey)]
-			{
-				if (TSharedPtr<FLiveLinkHubSession> Session = WeakSession.Pin())
-				{
-					Session->OnSubjectAdded(Key);
-				}
-			});
-	}
-
-	/** Handles updating the tree view when a subject is added. */
-	void OnSubjectAdded(const FLiveLinkSubjectKey& SubjectKey)
-	{
-		if (!SessionData->SubjectsConfig->SubjectProxies.Contains(SubjectKey))
-		{
-			ILiveLinkClient& LiveLinkClient = IModularFeatures::Get().GetModularFeature<ILiveLinkClient>(ILiveLinkClient::ModularFeatureName);
-
-			FLiveLinkHubSubjectProxy SubjectSettings;
-			SubjectSettings.Initialize(SubjectKey, LiveLinkClient.GetSourceType(SubjectKey.Source).ToString());
-
-			{
-				FWriteScopeLock Locker(SessionDataLock);
-				SessionData->SubjectsConfig->SubjectProxies.FindOrAdd(SubjectKey) = MoveTemp(SubjectSettings);
-
-				ULiveLinkHubSubjectProcessors* SubjectProcessors = NewObject<ULiveLinkHubSubjectProcessors>(SessionData->SubjectsConfig);
-				SubjectProcessors->Initialize(Cast<ULiveLinkSubjectSettings>(LiveLinkClient.GetSubjectSettings(SubjectKey)));
-
-				SessionData->SubjectsConfig->SubjectProcessors.FindOrAdd(SubjectKey) = SubjectProcessors;
-			}
-		}
-	}
-
-	/** AnyThread handler for the SubjectRemoved delegate, dispatches handling on the game thread to avoid asserts in Slate. */
-	void OnSubjectRemoved_AnyThread(FLiveLinkSubjectKey SubjectKey)
-	{
-		TWeakPtr<FLiveLinkHubSession> WeakSession = AsShared();
-		AsyncTask(ENamedThreads::GameThread, [WeakSession, Key = MoveTemp(SubjectKey)]
-			{
-				if (TSharedPtr<FLiveLinkHubSession> Session = WeakSession.Pin())
-				{
-					Session->OnSubjectRemoved(Key);
-				}
-			});
-	}
-
-	/** Handles updating the tree view when a subject is removed. */
-	void OnSubjectRemoved(const FLiveLinkSubjectKey& SubjectKey)
-	{
-		FWriteScopeLock Locker(SessionDataLock);
-		SessionData->SubjectsConfig->SubjectProxies.Remove(SubjectKey);
-		SessionData->SubjectsConfig->SubjectProcessors.Remove(SubjectKey);
-	}
-
 private:
 	/** List of clients in the current session. These represent the unreal instances than can receive data from the hub. */
 	TSet<FLiveLinkHubClientId> CachedSessionClients;
diff --git a/Engine/Plugins/Animation/LiveLinkHub/Source/LiveLinkHub/Private/Session/LiveLinkHubSessionData.h b/Engine/Plugins/Animation/LiveLinkHub/Source/LiveLinkHub/Private/Session/LiveLinkHubSessionData.h
--- a/Engine/Plugins/Animation/LiveLinkHub/Source/LiveLinkHub/Private/Session/LiveLinkHubSessionData.h
+++ b/Engine/Plugins/Animation/LiveLinkHub/Source/LiveLinkHub/Private/Session/LiveLinkHubSessionData.h
@@ -4,10 +4,9 @@
 
 #include "Clients/LiveLinkHubUEClientInfo.h"
 #include "CoreTypes.h"
-#include "Engine/TimecodeProvider.h"
+#include "LiveLinkHubMessages.h"
 #include "LiveLinkPresetTypes.h"
 #include "Misc/Guid.h"
-#include "Subjects/LiveLinkHubSubjectSessionConfig.h"
 
 
 #include "LiveLinkHubSessionData.generated.h"
@@ -19,21 +18,6 @@ class ULiveLinkHubSessionData : public UObject
 public:
 	GENERATED_BODY()
 
-	ULiveLinkHubSessionData()
-	{
-		SubjectsConfig = CreateDefaultSubobject<ULiveLinkHubSubjectSessionConfig>(TEXT("SubjectsConfig"));
-
-        if (!HasAnyFlags(RF_ClassDefaultObject | RF_ArchetypeObject))
-        {
-        	SubjectsConfig->Initialize();
-        }
-	}
-
-	ULiveLinkHubSessionData(ULiveLinkHubSubjectSessionConfig* InSubjectsConfig)
-	{
-		SubjectsConfig = InSubjectsConfig;
-	}
-
 	/** Live link hub sources. */
 	UPROPERTY()
 	TArray<FLiveLinkSourcePreset> Sources;
@@ -49,8 +33,4 @@ public:
 	/** Timecode settings for the live link hub. */
 	UPROPERTY()
 	FLiveLinkHubTimecodeSettings TimecodeSettings;
-
-	/** Subject configs for this session. */
-	UPROPERTY(Instanced)
-	TObjectPtr<ULiveLinkHubSubjectSessionConfig> SubjectsConfig;
 };
diff --git a/Engine/Plugins/Animation/LiveLinkHub/Source/LiveLinkHub/Private/Subjects/LiveLinkHubSubjectController.h b/Engine/Plugins/Animation/LiveLinkHub/Source/LiveLinkHub/Private/Subjects/LiveLinkHubSubjectController.h
--- a/Engine/Plugins/Animation/LiveLinkHub/Source/LiveLinkHub/Private/Subjects/LiveLinkHubSubjectController.h
+++ b/Engine/Plugins/Animation/LiveLinkHub/Source/LiveLinkHub/Private/Subjects/LiveLinkHubSubjectController.h
@@ -16,8 +16,6 @@ class FLiveLinkHubSubjectController
 public:
 	FLiveLinkHubSubjectController()
 	{
-		SubjectModel = MakeShared<FLiveLinkHubSubjectModel>();
-
 		const FLiveLinkHubModule& LiveLinkHubModule = FModuleManager::Get().GetModuleChecked<FLiveLinkHubModule>("LiveLinkHub");
 		LiveLinkHubModule.GetSessionManager()->OnActiveSessionChanged().AddRaw(this, &FLiveLinkHubSubjectController::OnActiveSessionChanged);
 	}
@@ -36,9 +34,7 @@ public:
 	/** Create the widget for displaying a subject's settings. */
 	TSharedRef<SWidget> MakeSubjectView()
 	{
-		return SAssignNew(SubjectsView, SLiveLinkHubSubjectView, SubjectModel.ToSharedRef())
-			.OnRenameSubject_Raw(this, &FLiveLinkHubSubjectController::OnSubjectRenamed)
-			.OnProcessorModified_Raw(this, &FLiveLinkHubSubjectController::OnSubjectProcessorModified);
+		return SAssignNew(SubjectsView, SLiveLinkHubSubjectView);
 	}
 
 	/** Set the displayed subject in the subject view. */
@@ -47,56 +43,6 @@ public:
 		SubjectsView->SetSubject(Subject);
 	}
 
-	/** Handle modifying the session config for the specified subject. */
-	void OnSubjectRenamed(const FLiveLinkSubjectKey& SubjectKey, FName NewName) const
-	{
-		if (TSharedPtr<ILiveLinkHubSessionManager> SessionManager = FModuleManager::Get().GetModuleChecked<FLiveLinkHubModule>("LiveLinkHub").GetLiveLinkHub()->GetSessionManager())
-		{
-			SessionManager->GetCurrentSession()->RenameSubject(SubjectKey, NewName);
-		}
-	}
-
-	void OnSubjectProcessorModified(const FLiveLinkSubjectKey& SubjectKey, const TArray<ULiveLinkFramePreProcessor*>& UpdatedPreprocessors, ULiveLinkFrameTranslator* UpdatedTranslator) const
-	{
-		FLiveLinkHubClient& LiveLinkClient = static_cast<FLiveLinkHubClient&>(IModularFeatures::Get().GetModularFeature<FLiveLinkClient>(ILiveLinkClient::ModularFeatureName));
-		ULiveLinkSubjectSettings* Settings = Cast<ULiveLinkSubjectSettings>(LiveLinkClient.GetSubjectSettings(SubjectKey));
-
-		Settings->PreProcessors = UpdatedPreprocessors;
-		Settings->Translators.Reset();
-
-		if (UpdatedTranslator)
-		{
-			Settings->Translators.Add(UpdatedTranslator);
-		}
-
-		// todo: If we are removing the translator, we will need to do additional handling to restore the original static data that was overriden
-		if (Settings->ValidateProcessors())
-		{
-			// Apply to the underlying data from the session.
-			if (const TSharedPtr<ILiveLinkHubSessionManager> SessionManager = FModuleManager::Get().GetModuleChecked<FLiveLinkHubModule>("LiveLinkHub").GetLiveLinkHub()->GetSessionManager())
-			{
-				SessionManager->GetCurrentSession()->SetPreProcessors(SubjectKey, Settings->PreProcessors);
-
-				if (Settings->Translators.Num())
-				{
-					SessionManager->GetCurrentSession()->SetTranslator(SubjectKey, Settings->Translators[0]);
-				}
-				else
-				{
-					SessionManager->GetCurrentSession()->SetTranslator(SubjectKey, nullptr);
-				}
-			}
-
-			// Apply to the livelink client
-			LiveLinkClient.CacheSubjectSettings(SubjectKey, Settings);
-		}
-		else
-		{
-			// Refresh the view since the underlying data might have been rolled back.
-			SubjectsView->SetSubject(SubjectKey);
-		}
-	}
-
 	/** Handle updating the subject details when the session has been swapped out for a different one. */
 	void OnActiveSessionChanged(const TSharedRef<ILiveLinkHubSession>& ActiveSession) const
 	{
@@ -109,6 +55,4 @@ public:
 private:
 	/** View widget for the selected subject. */
 	TSharedPtr<SLiveLinkHubSubjectView> SubjectsView;
-	/** Model responsible for the subjects data. */
-	TSharedPtr<ILiveLinkHubSubjectModel> SubjectModel;
 };
diff --git a/Engine/Plugins/Animation/LiveLinkHub/Source/LiveLinkHub/Private/Subjects/LiveLinkHubSubjectModel.h b/Engine/Plugins/Animation/LiveLinkHub/Source/LiveLinkHub/Private/Subjects/LiveLinkHubSubjectModel.h
deleted file mode 100644
diff --git a/Engine/Plugins/Animation/LiveLinkHub/Source/LiveLinkHub/Private/Subjects/LiveLinkHubSubjectSessionConfig.cpp b/Engine/Plugins/Animation/LiveLinkHub/Source/LiveLinkHub/Private/Subjects/LiveLinkHubSubjectSessionConfig.cpp
deleted file mode 100644
diff --git a/Engine/Plugins/Animation/LiveLinkHub/Source/LiveLinkHub/Private/Subjects/LiveLinkHubSubjectSessionConfig.h b/Engine/Plugins/Animation/LiveLinkHub/Source/LiveLinkHub/Private/Subjects/LiveLinkHubSubjectSessionConfig.h
deleted file mode 100644
diff --git a/Engine/Plugins/Animation/LiveLinkHub/Source/LiveLinkHub/Private/Subjects/SLiveLinkHubSubjectView.h b/Engine/Plugins/Animation/LiveLinkHub/Source/LiveLinkHub/Private/Subjects/SLiveLinkHubSubjectView.h
--- a/Engine/Plugins/Animation/LiveLinkHub/Source/LiveLinkHub/Private/Subjects/SLiveLinkHubSubjectView.h
+++ b/Engine/Plugins/Animation/LiveLinkHub/Source/LiveLinkHub/Private/Subjects/SLiveLinkHubSubjectView.h
@@ -2,25 +2,14 @@
 
 #pragma once
 
-#include "Async/Async.h"
-#include "IStructureDetailsView.h"
-#include "IPropertyRowGenerator.h"
-
-#include "Delegates/DelegateCombinations.h"
-#include "DetailsViewArgs.h"
 #include "Features/IModularFeatures.h"
-#include "IDetailsView.h"
-#include "ISinglePropertyView.h"
-#include "LiveLinkHubSubjectModel.h"
-#include "LiveLinkSubjectSettings.h"
+#include "LiveLinkClient.h"
+#include "ILiveLinkClient.h"
+#include "LiveLinkClientPanelViews.h"
 #include "LiveLinkTypes.h"
-#include "PropertyEditorModule.h"
 #include "Modules/ModuleManager.h"
+#include "SLiveLinkDataView.h"
 #include "Widgets/SCompoundWidget.h"
-#include "Widgets/SNullWidget.h"
-
-DECLARE_DELEGATE_TwoParams(FOnRenameLiveLinkHubSubject, const FLiveLinkSubjectKey& /*SubjectKey*/, FName /*NewName*/);
-DECLARE_DELEGATE_ThreeParams(FOnSubjectProcessorModified, const FLiveLinkSubjectKey& /*SubjectKey*/,  const TArray<ULiveLinkFramePreProcessor*>& /*UpdatedPreProcessors*/, ULiveLinkFrameTranslator* /*UpdatedTranslator*/);
 
 /**
  * Provides the UI that displays information about a livelink hub subject.
@@ -31,58 +20,24 @@ public:
 
 	SLATE_BEGIN_ARGS(SLiveLinkHubSubjectView) {}
 	SLATE_ARGUMENT(FLiveLinkSubjectKey, SubjectKey)
-	SLATE_EVENT(FOnRenameLiveLinkHubSubject, OnRenameSubject)
-	SLATE_EVENT(FOnSubjectProcessorModified, OnProcessorModified)
 	SLATE_END_ARGS()
 
 	//~ Begin SWidget interface
-	void Construct(const FArguments& InArgs, const TSharedRef<ILiveLinkHubSubjectModel>& InSubjectModel)
+	void Construct(const FArguments& InArgs)
 	{
-		SubjectModel = InSubjectModel;
-		OnRenameSubjectDelegate = InArgs._OnRenameSubject;
-		OnProcessorModifiedDelegate = InArgs._OnProcessorModified;
-
 		ILiveLinkClient& LiveLinkClient = static_cast<ILiveLinkClient&>(IModularFeatures::Get().GetModularFeature<FLiveLinkClient>(ILiveLinkClient::ModularFeatureName));
 		LiveLinkClient.OnLiveLinkSubjectRemoved().AddSP(this, &SLiveLinkHubSubjectView::OnSubjectRemoved);
 
-		FDetailsViewArgs DetailsViewArgs;
-		DetailsViewArgs.bUpdatesFromSelection = false;
-		DetailsViewArgs.bLockable = false;
-		DetailsViewArgs.bShowPropertyMatrixButton = false;
-		DetailsViewArgs.NameAreaSettings = FDetailsViewArgs::HideNameArea;
-		DetailsViewArgs.ViewIdentifier = NAME_None;
-		DetailsViewArgs.bShowCustomFilterOption = false;
-		DetailsViewArgs.bShowOptions = false;
-		DetailsViewArgs.bAllowSearch = false;
-
-		FPropertyEditorModule& PropertyEditorModule = FModuleManager::GetModuleChecked<FPropertyEditorModule>("PropertyEditor");
-		FStructureDetailsViewArgs StructureDetailsArgs;
-
-		SubjectData = MakeShared<TStructOnScope<FLiveLinkHubSubjectProxy>>();
-
-		SettingsDetailsView = PropertyEditorModule.CreateStructureDetailView(DetailsViewArgs, StructureDetailsArgs, SubjectData);
-		SettingsDetailsView->GetOnFinishedChangingPropertiesDelegate().AddSP(this, &SLiveLinkHubSubjectView::OnSubjectPropertyModified);
-
-		SettingsObjectDetailsView = PropertyEditorModule.CreateDetailView(DetailsViewArgs);
-		SettingsObjectDetailsView->OnFinishedChangingProperties().AddSP(this, &SLiveLinkHubSubjectView::OnSubjectPropertyModified);
+		constexpr bool bReadOnly = false;
+		SubjectView = UE::LiveLink::CreateSubjectsDetailsView(static_cast<FLiveLinkClient*>(&LiveLinkClient), bReadOnly);
 
 		ChildSlot
 		[
-			SNew(SVerticalBox)
-			+ SVerticalBox::Slot()
-			.AutoHeight()
-			[
-				SettingsDetailsView->GetWidget().ToSharedRef()
-			]
-			+ SVerticalBox::Slot()
-			.FillHeight(1.0f)
-			[
-				SettingsObjectDetailsView.ToSharedRef()
-			]
+			SubjectView.ToSharedRef()
 		];
 	}
 
-	~SLiveLinkHubSubjectView() override
+	virtual ~SLiveLinkHubSubjectView() override
 	{
 		if (IModularFeatures::Get().IsModularFeatureAvailable(ILiveLinkClient::ModularFeatureName))
 		{
@@ -92,92 +47,31 @@ public:
 	}
 
 	/** Clear the subject details. */
-	void RefreshSubjectDetails(const TSharedPtr<ILiveLinkHubSession>& ActiveSession)
+	void RefreshSubjectDetails(const TSharedPtr<class ILiveLinkHubSession>& ActiveSession)
 	{
-		SubjectData->Reset();
-		SettingsDetailsView->SetStructureData(nullptr);
-		SettingsObjectDetailsView->SetObject(nullptr);
-
-		SetSubject(SubjectKey);
+		SetSubject(FLiveLinkSubjectKey{});
 	}
 
 	/** Set the subject to be displayed in the details view. */
 	void SetSubject(const FLiveLinkSubjectKey& InSubjectKey)
 	{
 		SubjectKey = InSubjectKey;
-
-		if (TOptional<FLiveLinkHubSubjectProxy> SubjectProxy = SubjectModel->GetSubjectConfig(InSubjectKey))
-		{
-			const FLiveLinkHubSubjectProxy& Proxy = *SubjectProxy;
-			SubjectData->InitializeAs<FLiveLinkHubSubjectProxy>(*SubjectProxy);
-			SettingsDetailsView->SetStructureData(SubjectData);
-		}
-
-		if (ULiveLinkHubSubjectProcessors* SubjectProxy = SubjectModel->GetSubjectProcessors(InSubjectKey))
-		{
-			SettingsObjectDetailsView->SetObject(SubjectProxy);
-			WeakProcessors = SubjectProxy;
-		}
-	}
-
-	/** Get the outbound name for a subject. */
-	FString GetSubjectOutboundName(const FLiveLinkSubjectKey& InSubjectKey) const
-	{
-		if (TOptional<FLiveLinkHubSubjectProxy> SubjectProxy = SubjectModel->GetSubjectConfig(InSubjectKey))
-		{
-			return SubjectProxy->GetOutboundName().ToString();
-		}
-		return InSubjectKey.SubjectName.ToString();
-	}
-
-	/** Handler called when a subject property is modified, used to trigger a rename on the session's subject config. */
-	void OnSubjectPropertyModified(const FPropertyChangedEvent& PropertyChangedEvent)
-	{
-		if (SubjectData)
-		{
-			if (FLiveLinkHubSubjectProxy* Proxy = SubjectData->Get())
-			{
-				if (PropertyChangedEvent.GetPropertyName() == GET_MEMBER_NAME_CHECKED(FLiveLinkHubSubjectProxy, OutboundName))
-				{
-					OnRenameSubjectDelegate.ExecuteIfBound(SubjectKey, Proxy->GetOutboundName());
-				}
-			}
-		}
-
-		if (ULiveLinkHubSubjectProcessors* SubjectProcessors = WeakProcessors.Get())
-		{
-			if (PropertyChangedEvent.GetPropertyName() == GET_MEMBER_NAME_CHECKED(ULiveLinkHubSubjectProcessors, PreProcessors)
-				|| PropertyChangedEvent.GetPropertyName() == GET_MEMBER_NAME_CHECKED(ULiveLinkHubSubjectProcessors, Translator))
-			{
-				OnProcessorModifiedDelegate.ExecuteIfBound(SubjectKey, SubjectProcessors->PreProcessors, SubjectProcessors->Translator);
-			}
-		}
+		SubjectView->SetSubjectKey(SubjectKey);
 	}
 
 	void OnSubjectRemoved(FLiveLinkSubjectKey InSubjectKey)
 	{
 		if (InSubjectKey == SubjectKey)
 		{
-			SubjectKey = FLiveLinkSubjectKey();
-			RefreshSubjectDetails(nullptr);
+			SetSubject(FLiveLinkSubjectKey{});
 		}
 	}
 
 private:
-	/** Details for the selected subject. */
-	TSharedPtr<IStructureDetailsView> SettingsDetailsView;
 	/** Details for the selected subject. */
 	TSharedPtr<IDetailsView> SettingsObjectDetailsView;
 	/** Subject being shown. */
 	FLiveLinkSubjectKey SubjectKey;
-	/** Model that holds the data for a given subject. */
-	TSharedPtr<ILiveLinkHubSubjectModel> SubjectModel;
-	/** Struct on scope used to display subject data in the structure details view. */
-	TSharedPtr<TStructOnScope<FLiveLinkHubSubjectProxy>> SubjectData;
-	/** Delegate called when the outbound name is changed by the user. */
-	FOnRenameLiveLinkHubSubject OnRenameSubjectDelegate;
-	/** Delegate called when a translator or preprocessor is modified by the user. */
-	FOnSubjectProcessorModified OnProcessorModifiedDelegate;
-	/** Subject proxy object that holds preprocessor data. */
-	TWeakObjectPtr<ULiveLinkHubSubjectProcessors> WeakProcessors;
+
+	TSharedPtr<SLiveLinkDataView> SubjectView;
 };
