LiveLinkHub - Add details customization for livelinkhub subjects settings

diff --git a/Engine/Plugins/Animation/LiveLink/Source/LiveLinkEditor/Private/SLiveLinkDataView.cpp b/Engine/Plugins/Animation/LiveLink/Source/LiveLinkEditor/Private/SLiveLinkDataView.cpp
--- a/Engine/Plugins/Animation/LiveLink/Source/LiveLinkEditor/Private/SLiveLinkDataView.cpp
+++ b/Engine/Plugins/Animation/LiveLink/Source/LiveLinkEditor/Private/SLiveLinkDataView.cpp
@@ -7,6 +7,7 @@
 #include "IDetailsView.h"
 #include "LiveLinkClient.h"
 #include "LiveLinkRole.h"
+#include "LiveLinkSubjectSettings.h"
 
 #include "IStructureDetailsView.h"
 #include "Modules/ModuleManager.h"
@@ -205,7 +206,9 @@ void SLiveLinkDataView::Tick(const FGeometry& AllottedGeometry, const double InC
 void SLiveLinkDataView::SetSubjectKey(FLiveLinkSubjectKey InSubjectKey)
 {
 	SubjectKey = InSubjectKey;
-	SettingsDetailsView->SetObject(Client->GetSubjectSettings(SubjectKey));
+
+	UObject* SettingsObject = Client->GetSubjectSettings(SubjectKey);
+	SettingsDetailsView->SetObject(SettingsObject);
 }
 
 int32 SLiveLinkDataView::GetDetailWidgetIndex() const
diff --git a/Engine/Plugins/Animation/LiveLinkHub/Source/LiveLinkHub/Private/LiveLinkHubModule.cpp b/Engine/Plugins/Animation/LiveLinkHub/Source/LiveLinkHub/Private/LiveLinkHubModule.cpp
--- a/Engine/Plugins/Animation/LiveLinkHub/Source/LiveLinkHub/Private/LiveLinkHubModule.cpp
+++ b/Engine/Plugins/Animation/LiveLinkHub/Source/LiveLinkHub/Private/LiveLinkHubModule.cpp
@@ -5,9 +5,12 @@
 #include "Clients/LiveLinkHubProvider.h"
 #include "LiveLinkHubApplication.h"
 #include "LiveLinkHubLog.h"
+#include "LiveLinkHubSubjectSettings.h"
 #include "Modules/ModuleManager.h"
+#include "PropertyEditorModule.h"
 #include "Recording/LiveLinkHubPlaybackController.h"
 #include "Recording/LiveLinkHubRecordingController.h"
+#include "Subjects/LiveLinkHubSubjectSettingsDetailsCustomization.h"
 
 #if !WITH_LIVELINK_HUB
 #include "HAL/FileManager.h"
@@ -42,13 +45,22 @@ void FLiveLinkHubModule::StartupModule()
 		LOCTEXT("LiveLinkHubTooltip", "Launch the LiveLink Hub app."),
 		FSlateIcon("LiveLinkStyle", "LiveLinkClient.Common.Icon.Small"),
 		FUIAction(FExecuteAction::CreateRaw(this, &FLiveLinkHubModule::OpenLiveLinkHub)));
+#else
+	FPropertyEditorModule& PropertyModule = FModuleManager::LoadModuleChecked<FPropertyEditorModule>("PropertyEditor");
+	PropertyModule.RegisterCustomClassLayout(ULiveLinkHubSubjectSettings::StaticClass()->GetFName(), FOnGetDetailCustomizationInstance::CreateStatic(&FLiveLinkHubSubjectSettingsDetailsCustomization::MakeInstance));
 #endif
+
 }
 
 void FLiveLinkHubModule::ShutdownModule()
 {
 #if !WITH_LIVELINK_HUB
 	UToolMenus::UnregisterOwner(this);
+#else
+	if (FPropertyEditorModule* PropertyModule = FModuleManager::GetModulePtr<FPropertyEditorModule>("PropertyEditor"))
+	{
+		PropertyModule->UnregisterCustomClassLayout(ULiveLinkHubSubjectSettings::StaticClass()->GetFName());
+	}
 #endif
 }
 
diff --git a/Engine/Plugins/Animation/LiveLinkHub/Source/LiveLinkHub/Private/Subjects/LiveLinkHubSubjectSettingsDetailsCustomization.cpp b/Engine/Plugins/Animation/LiveLinkHub/Source/LiveLinkHub/Private/Subjects/LiveLinkHubSubjectSettingsDetailsCustomization.cpp
new file mode 100644
--- /dev/null
+++ b/Engine/Plugins/Animation/LiveLinkHub/Source/LiveLinkHub/Private/Subjects/LiveLinkHubSubjectSettingsDetailsCustomization.cpp
@@ -0,0 +1,26 @@
+// Copyright Epic Games, Inc. All Rights Reserved.
+
+#include "LiveLinkHubSubjectSettingsDetailsCustomization.h"
+
+#include "DetailLayoutBuilder.h"
+#include "LiveLinkHubSubjectSettings.h"
+
+#define LOCTEXT_NAMESPACE "LiveLinkHubSubjectSettingsDetailCustomization"
+
+void FLiveLinkHubSubjectSettingsDetailsCustomization::CustomizeDetails(IDetailLayoutBuilder& InDetailBuilder)
+{
+	//Get the current settings object being edited
+	TArray<TWeakObjectPtr<UObject>> ObjectsBeingCustomized;
+	InDetailBuilder.GetObjectsBeingCustomized(ObjectsBeingCustomized);
+	if (ObjectsBeingCustomized.Num() != 1)
+	{
+		return;
+	}
+
+	InDetailBuilder.HideProperty(GET_MEMBER_NAME_CHECKED(ULiveLinkHubSubjectSettings, InterpolationProcessor), ULiveLinkSubjectSettings::StaticClass());
+	InDetailBuilder.HideProperty(GET_MEMBER_NAME_CHECKED(ULiveLinkHubSubjectSettings, FrameRate), ULiveLinkSubjectSettings::StaticClass());
+	InDetailBuilder.HideProperty(GET_MEMBER_NAME_CHECKED(ULiveLinkHubSubjectSettings, bRebroadcastSubject), ULiveLinkSubjectSettings::StaticClass());
+	InDetailBuilder.HideProperty(GET_MEMBER_NAME_CHECKED(ULiveLinkHubSubjectSettings, Translators), ULiveLinkSubjectSettings::StaticClass());
+}
+
+#undef LOCTEXT_NAMESPACE
diff --git a/Engine/Plugins/Animation/LiveLinkHub/Source/LiveLinkHub/Private/Subjects/LiveLinkHubSubjectSettingsDetailsCustomization.h b/Engine/Plugins/Animation/LiveLinkHub/Source/LiveLinkHub/Private/Subjects/LiveLinkHubSubjectSettingsDetailsCustomization.h
new file mode 100644
--- /dev/null
+++ b/Engine/Plugins/Animation/LiveLinkHub/Source/LiveLinkHub/Private/Subjects/LiveLinkHubSubjectSettingsDetailsCustomization.h
@@ -0,0 +1,27 @@
+// Copyright Epic Games, Inc. All Rights Reserved.
+
+#pragma once
+
+#include "IDetailCustomization.h"
+
+class ULiveLinkHubSubjectSettings;
+
+
+/**
+* Customizes a ULiveLinkHubSubjectSettings object
+*/
+class FLiveLinkHubSubjectSettingsDetailsCustomization : public IDetailCustomization
+{
+public:
+
+	static TSharedRef<IDetailCustomization> MakeInstance()
+	{
+		return MakeShared<FLiveLinkHubSubjectSettingsDetailsCustomization>();
+	}
+
+	// IDetailCustomization interface
+	virtual void CustomizeDetails(class IDetailLayoutBuilder& DetailBuilder) override;
+	// End IDetailCustomization interface
+
+private:
+};
