LiveLinkHub: Support asynchronous file streaming when playing a recording.

diff --git a/Engine/Plugins/Animation/LiveLinkHub/Source/LiveLinkHub/Private/LiveLinkHubSettings.h b/Engine/Plugins/Animation/LiveLinkHub/Source/LiveLinkHub/Private/LiveLinkHubSettings.h
--- a/Engine/Plugins/Animation/LiveLinkHub/Source/LiveLinkHub/Private/LiveLinkHubSettings.h
+++ b/Engine/Plugins/Animation/LiveLinkHub/Source/LiveLinkHub/Private/LiveLinkHubSettings.h
@@ -19,6 +19,10 @@ public:
 	UPROPERTY(config, EditAnywhere, Category="LiveLinkHub", DisplayName = "Automatically add discovered clients")
 	bool bAutoAddDiscoveredClients = true;
 
+	/** The size in megabytes to buffer when streaming a recoding. */
+	UPROPERTY(config, EditAnywhere, Category="LiveLinkHub", meta = (ClampMin = "1", UIMin = "1"))
+	int32 FrameBufferSizeMB = 100;
+	
 	/** Which project settings sections to display when opening the settings viewer. */
 	UPROPERTY(config)
 	TArray<FName> ProjectSettingsToDisplay;
diff --git a/Engine/Plugins/Animation/LiveLinkHub/Source/LiveLinkHub/Private/Recording/Implementations/LiveLinkUAssetRecorder.cpp b/Engine/Plugins/Animation/LiveLinkHub/Source/LiveLinkHub/Private/Recording/Implementations/LiveLinkUAssetRecorder.cpp
--- a/Engine/Plugins/Animation/LiveLinkHub/Source/LiveLinkHub/Private/Recording/Implementations/LiveLinkUAssetRecorder.cpp
+++ b/Engine/Plugins/Animation/LiveLinkHub/Source/LiveLinkHub/Private/Recording/Implementations/LiveLinkUAssetRecorder.cpp
@@ -41,8 +41,10 @@ namespace UAssetRecorderUtils
 		
 		if (const FLiveLinkStaticDataStruct* StaticData = LiveLinkClient->GetSubjectStaticData_AnyThread(SubjectKey))
 		{
-			FInstancedStruct StaticDataInstancedStruct;
-			StaticDataInstancedStruct.InitializeAs(StaticData->GetStruct(), (uint8*)StaticData->GetBaseData());
+			check(StaticData->IsValid());
+		
+			TSharedPtr<FInstancedStruct> StaticDataInstancedStruct = MakeShared<FInstancedStruct>();
+			StaticDataInstancedStruct->InitializeAs(StaticData->GetStruct(), (uint8*)StaticData->GetBaseData());
 
 			StaticDataContainer = FLiveLinkRecordingStaticDataContainer();
 			StaticDataContainer->Role = LiveLinkRole;
@@ -82,7 +84,7 @@ bool FLiveLinkUAssetRecorder::IsRecording() const
 	return bIsRecording;
 }
 
-void FLiveLinkUAssetRecorder::RecordBaseData(FLiveLinkRecordingBaseDataContainer& StaticDataContainer, FInstancedStruct&& DataToRecord)
+void FLiveLinkUAssetRecorder::RecordBaseData(FLiveLinkRecordingBaseDataContainer& StaticDataContainer, TSharedPtr<FInstancedStruct>&& DataToRecord)
 {
 	const double TimeNowInSeconds = FPlatformTime::Seconds();
 	StaticDataContainer.RecordedData.Add(MoveTemp(DataToRecord));
@@ -94,8 +96,8 @@ void FLiveLinkUAssetRecorder::RecordStaticData(const FLiveLinkSubjectKey& Subjec
 	if (bIsRecording && CurrentRecording)
 	{
 		FLiveLinkRecordingStaticDataContainer& StaticDataContainer = CurrentRecording->StaticData.FindOrAdd(SubjectKey);
-		FInstancedStruct NewData;
-		NewData.InitializeAs(StaticData.GetStruct(), (uint8*)StaticData.GetBaseData());
+		TSharedPtr<FInstancedStruct> NewData = MakeShared<FInstancedStruct>();
+		NewData->InitializeAs(StaticData.GetStruct(), (uint8*)StaticData.GetBaseData());
 		StaticDataContainer.Role = Role;
 
 		RecordBaseData(StaticDataContainer, MoveTemp(NewData));
@@ -107,8 +109,8 @@ void FLiveLinkUAssetRecorder::RecordFrameData(const FLiveLinkSubjectKey& Subject
 	if (bIsRecording && CurrentRecording)
 	{
 		FLiveLinkRecordingBaseDataContainer& FrameDataContainer = CurrentRecording->FrameData.FindOrAdd(SubjectKey);
-		FInstancedStruct NewData;
-		NewData.InitializeAs(FrameData.GetStruct(), (uint8*)FrameData.GetBaseData());
+		TSharedPtr<FInstancedStruct> NewData = MakeShared<FInstancedStruct>();
+		NewData->InitializeAs(FrameData.GetStruct(), (uint8*)FrameData.GetBaseData());
 
 		RecordBaseData(FrameDataContainer, MoveTemp(NewData));
 	}
diff --git a/Engine/Plugins/Animation/LiveLinkHub/Source/LiveLinkHub/Private/Recording/Implementations/LiveLinkUAssetRecorder.h b/Engine/Plugins/Animation/LiveLinkHub/Source/LiveLinkHub/Private/Recording/Implementations/LiveLinkUAssetRecorder.h
--- a/Engine/Plugins/Animation/LiveLinkHub/Source/LiveLinkHub/Private/Recording/Implementations/LiveLinkUAssetRecorder.h
+++ b/Engine/Plugins/Animation/LiveLinkHub/Source/LiveLinkHub/Private/Recording/Implementations/LiveLinkUAssetRecorder.h
@@ -33,7 +33,7 @@ private:
 	/** Create a recording package and save it. */
 	void SaveRecording();
 	/** Record data to a ULiveLinkRecording object. */
-	void RecordBaseData(FLiveLinkRecordingBaseDataContainer& StaticDataContainer, FInstancedStruct&& DataToRecord);
+	void RecordBaseData(FLiveLinkRecordingBaseDataContainer& StaticDataContainer, TSharedPtr<FInstancedStruct>&& DataToRecord);
 	/** Record initial data for all livelink subjects. (Useful when static data was sent before the recording started). */
 	void RecordInitialStaticData();
 
diff --git a/Engine/Plugins/Animation/LiveLinkHub/Source/LiveLinkHub/Private/Recording/Implementations/LiveLinkUAssetRecording.cpp b/Engine/Plugins/Animation/LiveLinkHub/Source/LiveLinkHub/Private/Recording/Implementations/LiveLinkUAssetRecording.cpp
new file mode 100644
--- /dev/null
+++ b/Engine/Plugins/Animation/LiveLinkHub/Source/LiveLinkHub/Private/Recording/Implementations/LiveLinkUAssetRecording.cpp
@@ -0,0 +1,575 @@
+// Copyright Epic Games, Inc. All Rights Reserved.
+
+#include "LiveLinkUAssetRecording.h"
+
+#include "HAL/IConsoleManager.h"
+#include "LiveLinkHubLog.h"
+#include "LiveLinkUAssetRecordingPlayer.h"
+#include "Misc/FileHelper.h"
+#include "Misc/Paths.h"
+#include "Serialization/BufferArchive.h"
+#include "Serialization/MemoryReader.h"
+#include "Serialization/ObjectAndNameAsStringProxyArchive.h"
+
+ULiveLinkUAssetRecording::~ULiveLinkUAssetRecording()
+{
+	UnloadRecording();
+}
+
+void ULiveLinkUAssetRecording::Serialize(FArchive& Ar)
+{
+	Super::Serialize(Ar);
+	
+	FString FilePath = GetRecordingDataFilePath();
+	
+	if (Ar.IsSaving())
+	{
+		// todo: async stream save
+		SaveRecording();
+	}
+}
+
+void ULiveLinkUAssetRecording::SaveRecording()
+{
+	// Saving frame data to a custom file, so it can be streamed in later.
+	const FString FilePath = GetRecordingDataFilePath();
+	FArchive* FileWriter = IFileManager::Get().CreateFileWriter(*FilePath);
+
+	int32 RecordingVersionToSave = RecordingVersion;
+	*FileWriter << RecordingVersionToSave;
+
+	// How much static data to expect.
+	int32 NumStaticData = RecordingData.StaticData.Num();
+	*FileWriter << NumStaticData;
+		
+	for (TTuple<FLiveLinkSubjectKey, FLiveLinkRecordingStaticDataContainer>& StaticData : RecordingData.StaticData)
+	{
+		SaveFrameData(FileWriter, StaticData.Key, StaticData.Value);
+	}
+	
+	// How much frame data to expect.
+	int32 NumFrameData = RecordingData.FrameData.Num();
+	*FileWriter << NumFrameData;
+		
+	for (TTuple<FLiveLinkSubjectKey, FLiveLinkRecordingBaseDataContainer>& FrameData : RecordingData.FrameData)
+	{
+		SaveFrameData(FileWriter, FrameData.Key, FrameData.Value);
+	}
+		
+	FileWriter->FlushCache();
+	FileWriter->Close();
+	delete FileWriter;
+}
+
+void ULiveLinkUAssetRecording::LoadRecording(int32 InInitialFrame, int32 InNumFramesToLoad)
+{
+	bCancelStream = false;
+	
+	int32 StartFrame = InInitialFrame - InNumFramesToLoad;
+	if (StartFrame < 0)
+	{
+		StartFrame = 0;
+	}
+	
+	// Additional buffer to each side, plus the initial frame.
+	InNumFramesToLoad = (InNumFramesToLoad * 2) + 1;
+
+	// Perform initial setup of the file reader.
+	if (!AsyncStreamTask.IsValid())
+	{
+		check(RecordingFileReader == nullptr);
+		
+		RecordingStartFrameFilePosition = 0;
+		
+		const FString FilePath = GetRecordingDataFilePath();
+		RecordingFileReader = IFileManager::Get().CreateFileReader(*FilePath);
+
+		if (RecordingFileReader == nullptr)
+		{
+			UE_LOG(LogLiveLinkHub, Error, TEXT("Failed to open file %s for reading."), *FilePath);
+			return;
+		}
+	}
+
+	EarliestFrameToStream = StartFrame;
+
+	if (InitialFrameToStream != InInitialFrame)
+	{
+		bStreamingFrameChange = true;
+	}
+	
+	InitialFrameToStream = InInitialFrame;
+	TotalFramesToStream = InNumFramesToLoad;
+
+	if (!AsyncStreamTask.IsValid())
+	{
+		AsyncStreamTask = MakeUnique<FAsyncTask<FLiveLinkStreamAsyncTask>>(this);
+		AsyncStreamTask->StartBackgroundTask();
+	}
+}
+
+void ULiveLinkUAssetRecording::UnloadRecording()
+{
+	bCancelStream = true;
+
+	if (AsyncStreamTask.IsValid())
+	{
+		if (!AsyncStreamTask->Cancel())
+		{
+			AsyncStreamTask->EnsureCompletion();
+		}
+		AsyncStreamTask.Reset();
+	}
+
+	bPerformedInitialLoad = false;
+
+	if (RecordingFileReader)
+	{
+		RecordingFileReader->Close();
+		delete RecordingFileReader;
+		RecordingFileReader = nullptr;
+	}
+
+	FrameDataSubjectKey.Reset();
+	RecordingStartFrameFilePosition = 0;
+	RecordingMaxFrames = 0;
+	FrameDiskSize = 0;
+	EarliestFrameToStream = 0;
+	InitialFrameToStream = 0;
+	TotalFramesToStream = 0;
+	
+	BufferedFrames = TRange<int32>(0, 0);
+	
+	for (TTuple<FLiveLinkSubjectKey, FLiveLinkRecordingStaticDataContainer>& StaticData : RecordingData.StaticData)
+	{
+		StaticData.Value.Timestamps.Empty();
+		StaticData.Value.RecordedData.Empty();
+	}
+
+	for (TTuple<FLiveLinkSubjectKey, FLiveLinkRecordingBaseDataContainer>& FrameData : RecordingData.FrameData)
+	{
+		FrameData.Value.Timestamps.Empty();
+		FrameData.Value.RecordedData.Empty();
+	}
+}
+
+void ULiveLinkUAssetRecording::WaitForBufferedFrames(int32 InMinFrame, int32 InMaxFrame)
+{
+	if (AsyncStreamTask.IsValid())
+	{
+		// Max frames isn't set until after the initial load.
+		while (!bPerformedInitialLoad)
+		{
+			FPlatformProcess::Sleep(0.002);
+		}
+
+		// Clamp the frame range to the max possible range. If the selection range extends the actual frame range
+		// then there would be nothing to load.
+		InMinFrame = FMath::Clamp(InMinFrame, 0, RecordingMaxFrames - 1);
+		InMaxFrame = FMath::Clamp(InMaxFrame, 0, RecordingMaxFrames - 1);
+
+		const int32 InTotalFrames = InMaxFrame - InMinFrame + 1;
+		const TRange<int32> InRange(InMinFrame, InMaxFrame);
+		
+		while (true)
+		{
+			TRange<int32> BufferedFramesLocal = BufferedFrames.load();
+			if (InTotalFrames > TotalFramesToStream
+			|| BufferedFramesLocal.Contains(InRange) || AsyncStreamTask->IsDone())
+			{
+				break;
+			}
+			
+			FPlatformProcess::Sleep(0.002);
+		}
+	}
+}
+
+void ULiveLinkUAssetRecording::CopyRecordingData(FLiveLinkPlaybackTracks& InOutLiveLinkPlaybackTracks) const
+{
+	FScopeLock Lock(&DataContainerMutex);
+
+	for (const TPair<FLiveLinkSubjectKey, FLiveLinkRecordingStaticDataContainer>& Pair : RecordingData.StaticData)
+	{
+		// Modify subject name so a duplicate FLiveLinkSubjectKey below doesn't produce the same hash. This allows us to efficiently
+		// reuse tracks, as well as preserve the absolute frame index, which is needed since frame data is streamed in.
+		FLiveLinkSubjectKey StaticSubjectKey = Pair.Key;
+		StaticSubjectKey.SubjectName.Name = *(StaticSubjectKey.SubjectName.ToString() + "_STATIC");
+		FLiveLinkPlaybackTrack& PlaybackTrack = InOutLiveLinkPlaybackTracks.Tracks.FindOrAdd(StaticSubjectKey);
+
+		PlaybackTrack.FrameData = TConstArrayView<TSharedPtr<FInstancedStruct>>(Pair.Value.RecordedData);
+		PlaybackTrack.Timestamps = TConstArrayView<double>(Pair.Value.Timestamps);
+		PlaybackTrack.LiveLinkRole = Pair.Value.Role;
+		PlaybackTrack.SubjectKey = Pair.Key;
+		PlaybackTrack.StartIndexOffset = Pair.Value.RecordedDataStartFrame;
+	}
+
+	for (const TPair<FLiveLinkSubjectKey, FLiveLinkRecordingBaseDataContainer>& Pair : RecordingData.FrameData)
+	{
+		FLiveLinkPlaybackTrack& PlaybackTrack = InOutLiveLinkPlaybackTracks.Tracks.FindOrAdd(Pair.Key);
+		
+		PlaybackTrack.FrameData = TConstArrayView<TSharedPtr<FInstancedStruct>>(Pair.Value.RecordedData);
+		PlaybackTrack.Timestamps = TConstArrayView<double>(Pair.Value.Timestamps);
+		PlaybackTrack.SubjectKey = Pair.Key;
+		PlaybackTrack.StartIndexOffset = Pair.Value.RecordedDataStartFrame;
+	}
+}
+
+void ULiveLinkUAssetRecording::SaveFrameData(FArchive* InFileWriter, const FLiveLinkSubjectKey& InSubjectKey, FLiveLinkRecordingBaseDataContainer& InBaseDataContainer)
+{
+		// This will crash if it fails -- we don't want to save invalid data.
+		InBaseDataContainer.ValidateData();
+		
+		// Start block with map key.
+		FGuid Source = InSubjectKey.Source;
+		FString SubjectName = InSubjectKey.SubjectName.ToString();
+		int32 NumFrames = InBaseDataContainer.RecordedData.Num();
+		*InFileWriter << Source;
+		*InFileWriter << SubjectName;
+		*InFileWriter << NumFrames;
+
+		uint64 SerializedFrameSizePosition = 0;
+		int32 SerializedFrameSize = 0;
+		if (NumFrames > 0)
+		{
+			const UScriptStruct* ScriptStruct = InBaseDataContainer.RecordedData[0]->GetScriptStruct();
+			FString StructTypeName = ScriptStruct->GetPathName();
+
+			// Write the struct name and size so it can be loaded later.
+			*InFileWriter << StructTypeName;
+
+			// Remember the position to write the frame size.
+			SerializedFrameSizePosition = InFileWriter->Tell();
+			*InFileWriter << SerializedFrameSize;
+		}
+			
+		for (int32 FrameIdx = 0; FrameIdx < NumFrames; ++FrameIdx)
+		{
+			TSharedPtr<FInstancedStruct>& Frame = InBaseDataContainer.RecordedData[FrameIdx];
+			check(Frame.IsValid() && Frame->IsValid());
+			
+			// Write the frame index for streaming frames when loading.
+			*InFileWriter << FrameIdx;
+
+			// Write the frame's timestamp.
+			double Timestamp = InBaseDataContainer.Timestamps[FrameIdx];
+			*InFileWriter << Timestamp;
+			
+			// Write the entire frame data.
+			uint64 SerializeDataStart = InFileWriter->Tell();
+			FObjectAndNameAsStringProxyArchive StructAr(*InFileWriter, false);
+			Frame->Serialize(StructAr);
+
+			// Store the serialized frame size, so we can write it once later.
+			{
+				int32 CurrentSerializedFrameSize = InFileWriter->Tell() - SerializeDataStart;
+				// Sanity check that the serialized frame size is consistent.
+				ensure(CurrentSerializedFrameSize == SerializedFrameSize || SerializedFrameSize == 0);
+				SerializedFrameSize = CurrentSerializedFrameSize;
+			}
+		}
+
+		if (SerializedFrameSize > 0)
+		{
+			// Write the frame data offset at the beginning of the block.
+			const uint64 FinalOffset = InFileWriter->Tell();
+			InFileWriter->Seek(SerializedFrameSizePosition);
+			*InFileWriter << SerializedFrameSize;
+			InFileWriter->Seek(FinalOffset);
+		}
+}
+
+void ULiveLinkUAssetRecording::LoadRecordingAsync(int32 InStartFrame, int32 InCurrentFrame, int32 InNumFramesToLoad)
+{
+	const int32 MaxPossibleFrame = RecordingMaxFrames - 1;
+	InStartFrame = FMath::Clamp(InStartFrame, 0, MaxPossibleFrame);
+	InCurrentFrame = FMath::Clamp(InCurrentFrame, 0, MaxPossibleFrame);
+	const int32 EndFrame = InStartFrame + InNumFramesToLoad - 1;
+
+	if (BufferedFrames.load().Contains(TRange<int32>(InStartFrame, FMath::Min(EndFrame, MaxPossibleFrame > 0 ? MaxPossibleFrame : EndFrame))))
+	{
+		// All frames are already buffered.
+		return;
+	}
+	
+	RecordingFileReader->Seek(RecordingStartFrameFilePosition);
+
+	// Read file for streaming into memory.
+	if (!RecordingFileReader->AtEnd())
+	{
+		// Perform initial load and record entry frame file offsets.
+		const bool bInitialLoad = RecordingStartFrameFilePosition == 0;
+		if (bInitialLoad)
+		{
+			int32 LoadedRecordingVersion;
+			*RecordingFileReader << LoadedRecordingVersion;
+
+			// If we modify the RecordingVersion we can perform import logic here.
+			ensure(LoadedRecordingVersion == RecordingVersion);
+			
+			// Process static data.
+		
+			int32 NumStaticData = 0;
+			*RecordingFileReader << NumStaticData;
+
+			if (NumStaticData > 0)
+			{
+				FGuid KeySource = FGuid();
+				FString KeyName;
+			
+				*RecordingFileReader << KeySource;
+				*RecordingFileReader << KeyName;
+			
+				FLiveLinkSubjectKey SubjectKey(KeySource, *KeyName);
+
+				FLiveLinkRecordingStaticDataContainer& DataContainer = RecordingData.StaticData.FindChecked(SubjectKey);
+				LoadFrameData(&DataContainer, 0, 0, 1);
+			}
+
+			// Process frame data.
+		
+			int32 NumFrameData = 0;
+			*RecordingFileReader << NumFrameData;
+
+			if (NumFrameData > 0)
+			{
+				FGuid KeySource = FGuid();
+				FString KeyName;
+			
+				*RecordingFileReader << KeySource;
+				*RecordingFileReader << KeyName;
+			
+				FrameDataSubjectKey = MakeShared<FLiveLinkSubjectKey>(KeySource, *KeyName);
+			}
+
+			RecordingStartFrameFilePosition = RecordingFileReader->Tell();
+		}
+
+		// Load the required frames, either on initial load or subsequent loads.
+		if (ensure(FrameDataSubjectKey.IsValid()))
+		{
+			FLiveLinkRecordingBaseDataContainer& DataContainer = RecordingData.FrameData.FindChecked(*FrameDataSubjectKey.Get());
+			LoadFrameData(&DataContainer, InStartFrame, InCurrentFrame, InNumFramesToLoad);
+		}
+		else
+		{
+			UE_LOG(LogLiveLinkHub, Error, TEXT("FrameDataSubjectKey is missing for file %s."), *GetRecordingDataFilePath());
+		}
+		bPerformedInitialLoad = true;
+	}
+}
+
+void ULiveLinkUAssetRecording::LoadFrameData(FLiveLinkRecordingBaseDataContainer* DataContainer, int32 RequestedStartFrame, int32 RequestedInitialFrame, int32 RequestedFramesToLoad)
+{
+	bStreamingFrameChange = false;
+
+	int32 MaxFrames = 0;
+	(*RecordingFileReader) << MaxFrames;
+
+	if (MaxFrames > RecordingMaxFrames)
+	{
+		RecordingMaxFrames = MaxFrames;
+	}
+			
+	if (MaxFrames > 0)
+	{
+		if (RequestedFramesToLoad > 0)
+		{
+			// Don't go past requested frames or max frames.
+			MaxFrames = FMath::Min(MaxFrames, RequestedStartFrame + RequestedFramesToLoad);
+		}
+				
+		FString StructTypeName;
+		int32 SerializedStructureSize;
+			
+		*RecordingFileReader << StructTypeName;
+		*RecordingFileReader << SerializedStructureSize;
+
+		int64 StartPosition = RecordingFileReader->Tell();
+				
+		UScriptStruct* LoadedStruct = FindObject<UScriptStruct>(nullptr, *StructTypeName, true);
+		if (LoadedStruct == nullptr)
+		{
+			UE_LOG(LogLiveLinkHub, Error, TEXT("Script struct type '%s' not found."), *StructTypeName);
+			return;
+		}
+
+		// The size on disk for each frame-- consisting of the frame index, timestamp, and frame struct data.
+		FrameDiskSize = (sizeof(int32) + sizeof(double) + SerializedStructureSize);
+		
+		auto GetFrameFilePosition = [StartPosition, this](int32 FrameIdx)
+		{
+			return StartPosition + (FrameDiskSize * FrameIdx);
+		};
+
+		// Seek to the requested start frame position in the file.
+		int64 NewPosition = GetFrameFilePosition(RequestedStartFrame);
+		RecordingFileReader->Seek(NewPosition);
+
+		// Arrays to store the newly loaded data.
+		TArray<double> NewTimestamps;
+		TArray<TSharedPtr<FInstancedStruct>> NewRecordedData;
+
+		NewTimestamps.Reserve(RequestedFramesToLoad);
+		NewRecordedData.Reserve(RequestedFramesToLoad);
+
+		// Load each frame from the initial frame, alternating right to left each frame. This creates a buffer to support
+		// scrubbing each direction and makes sure the immediate frames are loaded first.
+		
+		int32 RightFrameIdx = RequestedInitialFrame;
+		int32 LeftFrameIdx = RequestedInitialFrame - 1; // - 1 So we don't try to load the same initial frame when alternating to the left.
+		int32 LastLoadedRightFrame = RequestedInitialFrame;
+		int32 LastLoadedLeftFrame = RequestedInitialFrame;
+		bool bLoadRight = true; // Start right -> left
+		
+		// We could potentially optimize this further -- such as adjusting the ratio of ahead/behind frames to buffer based on whether
+		// the recording is playing forward or reverse vs being scrubbed.
+
+		auto AlternateLoadDirection = [&]()
+		{
+			bLoadRight = !bLoadRight;
+		};
+		
+		while (RightFrameIdx < MaxFrames || LeftFrameIdx >= RequestedStartFrame)
+		{
+			int32 FrameToLoad;
+    
+			if (bLoadRight)
+			{
+				if (RightFrameIdx >= MaxFrames)
+				{
+					bLoadRight = false;
+					continue;
+				}
+				FrameToLoad = RightFrameIdx++;
+				LastLoadedRightFrame = FrameToLoad;
+			}
+			else
+			{
+				if (LeftFrameIdx < RequestedStartFrame)
+				{
+					bLoadRight = true;
+					continue;
+				}
+				FrameToLoad = LeftFrameIdx--;
+				LastLoadedLeftFrame = FrameToLoad;
+			}
+
+			int64 FramePosition = GetFrameFilePosition(FrameToLoad);
+			RecordingFileReader->Seek(FramePosition);
+
+			int32 ParsedFrameIdx = 0;
+			*RecordingFileReader << ParsedFrameIdx;
+
+			// Ensure the parsed frame index matches the expected frame
+			if (!ensure(ParsedFrameIdx == FrameToLoad))
+			{
+				UE_LOG(LogLiveLinkHub, Error, TEXT("Frame index mismatch: expected %d, got %d"), FrameToLoad, ParsedFrameIdx);
+				break;
+			}
+
+			auto InsertFrame = [&](TSharedPtr<FInstancedStruct>& InFrame, double InTimestamp, bool bCopy)
+			{
+#if UE_BUILD_DEBUG
+				ensure(!NewTimestamps.Contains(InTimestamp));
+#endif
+				if (bLoadRight)
+				{
+					NewTimestamps.Add(InTimestamp);
+					NewRecordedData.Add(bCopy ? InFrame : MoveTemp(InFrame));
+				}
+				else
+				{
+					NewTimestamps.Insert(InTimestamp, 0);
+					NewRecordedData.Insert(bCopy ? InFrame : MoveTemp(InFrame), 0);
+				}
+				
+#if UE_BUILD_DEBUG
+				// Additional validation to ensure timestamps / frames are loaded in the correct order.
+				for (int32 Idx = 1; Idx < NewTimestamps.Num(); ++Idx)
+				{
+					double LastTimestamp = NewTimestamps[Idx - 1];
+					double CurrentTimestamp = NewTimestamps[Idx];
+					ensure(LastTimestamp < CurrentTimestamp);
+				}
+#endif
+			};
+			
+			// Don't load a frame already in memory.
+			double ExistingTimestamp = 0.f;
+			if (TSharedPtr<FInstancedStruct> ExistingFrame = DataContainer->TryGetFrame(ParsedFrameIdx, ExistingTimestamp))
+			{
+				check(RecordingStartFrameFilePosition != 0);
+				constexpr bool bCopy = true; // Copy, as the data container could still be having its frame data pushed to animation.
+				InsertFrame(ExistingFrame, ExistingTimestamp, bCopy);
+				AlternateLoadDirection();
+				continue;
+			}
+
+			double Timestamp = 0;
+			*RecordingFileReader << Timestamp;
+
+			// Instantiate the animation frame.
+			FObjectAndNameAsStringProxyArchive StructAr(*RecordingFileReader, false);
+			TSharedPtr<FInstancedStruct> DataInstancedStruct = MakeShared<FInstancedStruct>(LoadedStruct);
+			DataInstancedStruct->Serialize(StructAr);
+
+			InsertFrame(DataInstancedStruct, Timestamp, false);
+		
+			ensure(NewTimestamps.Num() == NewRecordedData.Num());
+
+			AlternateLoadDirection();
+			
+			if (bStreamingFrameChange)
+			{
+				// The requested frames to stream have changed, finish the cycle and let the async task continue with updated data.
+				break;
+			}
+		}
+
+		// Record all the frames that have been buffered.
+		BufferedFrames = TRange<int32>(LastLoadedLeftFrame, LastLoadedRightFrame);
+
+		// Output the streamed data to the data container. This will unload unused frames.
+		{
+			FScopeLock Lock(&DataContainerMutex);
+			DataContainer->Timestamps = MoveTemp(NewTimestamps);
+			DataContainer->RecordedData = MoveTemp(NewRecordedData);
+			DataContainer->RecordedDataStartFrame = LastLoadedLeftFrame;
+
+			// This could potentially be optimized, such as outputting directly to the data container during the iteration, which could
+			// allow smoother streaming when scrubbing to a position that isn't buffered at all. However, we would need to be careful of the cost
+			// of locking the container each iteration.
+		}
+	}
+}
+
+FString ULiveLinkUAssetRecording::GetRecordingDataFilePath() const
+{
+	// todo: This path needs to be calculated from the settings file.
+	FString AssetPath = GetPathName();
+	FString BasePath = FPaths::ProjectSavedDir();
+	FString FilePath = FString::Printf(TEXT("%sRecordingData/%s.rec"), *BasePath, *FPaths::GetBaseFilename(AssetPath));
+
+	return FilePath;
+}
+
+void ULiveLinkUAssetRecording::FLiveLinkStreamAsyncTask::DoWork()
+{
+	int32 LastStartFrame = -1;
+	int32 LastTotalFrames = -1;
+	int32 LastInitialFrame = -1;
+	while (LiveLinkRecording && !LiveLinkRecording->bCancelStream)
+	{
+		if (LastStartFrame != LiveLinkRecording->EarliestFrameToStream
+			|| LastTotalFrames != LiveLinkRecording->TotalFramesToStream
+			|| LastInitialFrame != LiveLinkRecording->InitialFrameToStream)
+		{
+			LastStartFrame = LiveLinkRecording->EarliestFrameToStream;
+			LastTotalFrames = LiveLinkRecording->TotalFramesToStream;
+			LastInitialFrame = LiveLinkRecording->InitialFrameToStream;
+			LiveLinkRecording->LoadRecordingAsync(LiveLinkRecording->EarliestFrameToStream,
+				LiveLinkRecording->InitialFrameToStream, LiveLinkRecording->TotalFramesToStream);
+		}
+	}
+}
diff --git a/Engine/Plugins/Animation/LiveLinkHub/Source/LiveLinkHub/Private/Recording/Implementations/LiveLinkUAssetRecording.h b/Engine/Plugins/Animation/LiveLinkHub/Source/LiveLinkHub/Private/Recording/Implementations/LiveLinkUAssetRecording.h
--- a/Engine/Plugins/Animation/LiveLinkHub/Source/LiveLinkHub/Private/Recording/Implementations/LiveLinkUAssetRecording.h
+++ b/Engine/Plugins/Animation/LiveLinkHub/Source/LiveLinkHub/Private/Recording/Implementations/LiveLinkUAssetRecording.h
@@ -2,6 +2,7 @@
 
 #pragma once
 
+#include "Async/AsyncWork.h"
 #include "InstancedStruct.h"
 #include "LiveLinkPreset.h"
 #include "LiveLinkRole.h"
@@ -13,19 +14,84 @@
 
 #include "LiveLinkUAssetRecording.generated.h"
 
+struct FLiveLinkPlaybackTracks;
+
 /** Base data container for a recording track. */
 USTRUCT()
 struct FLiveLinkRecordingBaseDataContainer
 {
 	GENERATED_BODY()
 
-	/** Timestamps for the recorded data. Each entry matches an entry in the RecordedData array. */
-	UPROPERTY()
+	/** SERIALIZED DATA - Timestamps for the recorded data. Each entry matches an entry in the RecordedData array. */
 	TArray<double> Timestamps;
+	
+	/**
+	 * SERIALIZED DATA - Array of either static or frame recorded for a given timestamp.
+	 * TSharedPtr used as streaming the data in may require shared access.
+	 * FSharedStruct is not used because it doesn't implement Serialize().
+	 */
+	TArray<TSharedPtr<FInstancedStruct>> RecordedData;
 
-	/** Array of either static or frame recorded for a given timestamp.*/
-	UPROPERTY()
-	TArray<FInstancedStruct> RecordedData;
+	/** The current start frame for RecordedData. */
+	int32 RecordedDataStartFrame = 0;
+	
+	/**
+	 * Retrieve a loaded frame.
+	 *
+	 * @param InFrame The absolute frame index to load.
+	 *
+	 * @return The frame if it exists.
+	 */
+	TSharedPtr<FInstancedStruct> TryGetFrame(const int32 InFrame) const
+	{
+		if (IsFrameLoaded(InFrame))
+		{
+			const int32 RelativeFrameIdx = InFrame - RecordedDataStartFrame;
+			return RecordedData[RelativeFrameIdx];
+		}
+		return nullptr;
+	}
+
+	/**
+	 * Retrieve a loaded frame.
+	 *
+	 * @param InFrame The absolute frame index to load.
+	 * @param OutTimestamp Return the timestamp, if it exists.
+	 *
+	 * @return The frame if it exists.
+	 */
+	TSharedPtr<FInstancedStruct> TryGetFrame(const int32 InFrame, double& OutTimestamp) const
+	{
+		if (TSharedPtr<FInstancedStruct> OutFrame = TryGetFrame(InFrame))
+		{
+			const int32 RelativeFrameIdx = InFrame - RecordedDataStartFrame;
+			OutTimestamp = Timestamps[RelativeFrameIdx];
+			return OutFrame;
+		}
+		return nullptr;
+	}
+	
+	/**
+	 * Checks if a frame is currently loaded.
+	 * 
+	 * @param InFrame The absolute frame index to check.
+	 *
+	 * @return True if the frame index exists within the frame array.
+	 */
+	bool IsFrameLoaded(const int32 InFrame) const
+	{
+		return InFrame >= RecordedDataStartFrame && InFrame < RecordedDataStartFrame + RecordedData.Num();
+	}
+
+	/** Check data memory is valid and expected. */
+	void ValidateData() const
+	{
+		check(Timestamps.Num() == RecordedData.Num());
+		for (const TSharedPtr<FInstancedStruct>& InstancedStruct : RecordedData)
+		{
+			check(InstancedStruct.IsValid() && InstancedStruct->IsValid());
+		}
+	}
 };
 
 /** Container for static data. */
@@ -63,7 +129,117 @@ class ULiveLinkUAssetRecording : public ULiveLinkRecording
 public:
 	GENERATED_BODY()
 
-	/** Recorded Static and frame data. */
+	virtual ~ULiveLinkUAssetRecording() override;
+	
+	virtual void Serialize(FArchive& Ar) override;
+
+	/** Save recording data to disk. */
+	void SaveRecording();
+
+	/** Load recording data from disk. */
+	void LoadRecording(int32 InInitialFrame, int32 InNumFramesToLoad);
+
+	/** Free memory and close file reader. */
+	void UnloadRecording();
+
+	/** Block until frames are loaded. */
+	void WaitForBufferedFrames(int32 InMinFrame, int32 InMaxFrame);
+	
+	/** Return the maximum frames for this recording. */
+	int32 GetMaxFrames() const { return RecordingMaxFrames; }
+
+	/** The size in bytes of each animation frame. */
+	int32 GetFrameDiskSize() const { return FrameDiskSize; }
+	
+	/** Return the currently buffered frame range. */
+	TRange<int32> GetBufferedFrames() const { return BufferedFrames.load(); }
+
+	/** Copy the asset's loaded recording data to a format suitable for playback in live link. */
+	void CopyRecordingData(FLiveLinkPlaybackTracks& InOutLiveLinkPlaybackTracks) const;
+	
+private:
+	/** Serialize the number of frames (array size) of the BaseDataContainer to the archive. */
+	void SaveFrameData(FArchive* InFileWriter, const FLiveLinkSubjectKey& InSubjectKey, FLiveLinkRecordingBaseDataContainer& InBaseDataContainer);
+	
+	/** Initialize or update an async load. */
+	void LoadRecordingAsync(int32 InStartFrame, int32 InCurrentFrame, int32 InNumFramesToLoad);
+	
+	/** Load frame data to a data container. */
+	void LoadFrameData(FLiveLinkRecordingBaseDataContainer* DataContainer,
+		int32 RequestedStartFrame, int32 RequestedInitialFrame, int32 RequestedFramesToLoad);
+	
+	/** Retrieve the recording data file path for this asset. */
+	FString GetRecordingDataFilePath() const;
+	
+public:
+	/** Recorded static and frame data. */
 	UPROPERTY()
 	FLiveLinkUAssetRecordingData RecordingData;
+	
+private:
+	class FLiveLinkStreamAsyncTask : public FNonAbandonableTask
+	{
+	public:
+		FLiveLinkStreamAsyncTask(ULiveLinkUAssetRecording* InLiveLinkRecording)
+		{
+			LiveLinkRecording = InLiveLinkRecording;
+		}
+
+		TStatId GetStatId() const
+		{
+			RETURN_QUICK_DECLARE_CYCLE_STAT(LiveLinkStreamAsyncTask, STATGROUP_ThreadPoolAsyncTasks);
+		}
+
+		void DoWork();
+
+	private:
+		TObjectPtr<ULiveLinkUAssetRecording> LiveLinkRecording;
+	};
+
+	friend class FLiveLinkStreamAsyncTask;
+	
+	/** The file reader for the recording data. */
+	FArchive* RecordingFileReader = nullptr;
+
+	/** The subject key used for the frame data. */
+	TSharedPtr<FLiveLinkSubjectKey> FrameDataSubjectKey;
+	
+	/** The position in the file recording where frame data begins. */
+	int64 RecordingStartFrameFilePosition = 0;
+	
+	/** The maximum frames for this recording. */
+	int32 RecordingMaxFrames = 0;
+	
+	/** The first (left most) frame to stream. */
+	int32 EarliestFrameToStream = 0;
+
+	/** The initial frame to start streaming (the current playhead position).  */
+	int32 InitialFrameToStream = 0;
+
+	/** Total frames which should be streamed. */
+	int32 TotalFramesToStream = 0;
+
+	/** When the streaming frame has changed, signalling the current stream task should restart. */
+	std::atomic<bool> bStreamingFrameChange = false;
+
+	/** Signal that the stream should be canceled. */
+	std::atomic<bool> bCancelStream = false;
+
+	/** True once a full initial load has been performed -- static + frame data. */
+	std::atomic<bool> bPerformedInitialLoad = false;
+	
+	/** The size in bytes of each animation frame. */
+	std::atomic<int32> FrameDiskSize = 0;
+
+	/** Frames that are loaded and buffered. */
+	std::atomic<TRange<int32>> BufferedFrames = TRange<int32>(0, 0);
+
+	/** Mutex for accessing the data container from multiple threads. */
+	mutable FCriticalSection DataContainerMutex;
+
+	/** The thread streaming data from disk. */
+	TUniquePtr<FAsyncTask<FLiveLinkStreamAsyncTask>> AsyncStreamTask;
+
+	/** The current version of the recording. */
+	const int32 RecordingVersion = 1;
 };
diff --git a/Engine/Plugins/Animation/LiveLinkHub/Source/LiveLinkHub/Private/Recording/Implementations/LiveLinkUAssetRecordingPlayer.cpp b/Engine/Plugins/Animation/LiveLinkHub/Source/LiveLinkHub/Private/Recording/Implementations/LiveLinkUAssetRecordingPlayer.cpp
--- a/Engine/Plugins/Animation/LiveLinkHub/Source/LiveLinkHub/Private/Recording/Implementations/LiveLinkUAssetRecordingPlayer.cpp
+++ b/Engine/Plugins/Animation/LiveLinkHub/Source/LiveLinkHub/Private/Recording/Implementations/LiveLinkUAssetRecordingPlayer.cpp
@@ -2,7 +2,9 @@
 
 #include "LiveLinkUAssetRecordingPlayer.h"
 
+#include "HAL/IConsoleManager.h"
 #include "LiveLinkHubLog.h"
+#include "LiveLinkHubSettings.h"
 #include "Recording/LiveLinkRecording.h"
 
 
@@ -36,7 +38,8 @@ public:
 
 	const FInstancedStruct& FrameData() const
 	{
-		return Track.FrameData[FrameIndex];
+		check(FrameIndex >= 0 && FrameIndex < Track.FrameData.Num());
+		return *Track.FrameData[FrameIndex];
 	}
 
 	int32 CurrentIndex() const
@@ -104,20 +107,31 @@ private:
 
 void FLiveLinkPlaybackTrack::GetFramesUntil(double InPlayhead, TArray<FLiveLinkRecordedFrame>& OutFrames)
 {
-	for (FLiveLinkPlaybackTrackForwardIterator It = FLiveLinkPlaybackTrackForwardIterator(*this, LastReadIndex + 1); It; ++It)
+	for (FLiveLinkPlaybackTrackForwardIterator It = FLiveLinkPlaybackTrackForwardIterator(*this, GetRelativeIndex(LastReadAbsoluteIndex)); It; ++It)
 	{
-		if (It.FrameTimestamp() > InPlayhead)
+		const double FrameTimestamp = It.FrameTimestamp();
+		if (FrameTimestamp == LastTimeStamp)
+		{
+			// Generally the first iteration from LastReadAbsoluteIndex will trigger this,
+			// but it's possible if the maximum buffered frames are small (ie 1), then
+			// the LastReadAbsoluteIndex will now point to a different frame.
+			continue;
+		}
+		
+		if (FrameTimestamp > InPlayhead)
 		{
 			break;
 		}
 
-		LastReadIndex = It.CurrentIndex();
-
+		LastReadRelativeIndex = It.CurrentIndex();
+		LastReadAbsoluteIndex = LastReadRelativeIndex + StartIndexOffset;
+		LastTimeStamp = FrameTimestamp;
+		
 		FLiveLinkRecordedFrame FrameToPlay;
 		FrameToPlay.Data = It.FrameData();
 		FrameToPlay.SubjectKey = SubjectKey;
 		FrameToPlay.LiveLinkRole = LiveLinkRole;
-		FrameToPlay.FrameIndex = LastReadIndex;
+		FrameToPlay.FrameIndex = LastReadAbsoluteIndex;
 
 		OutFrames.Add(MoveTemp(FrameToPlay));
 	}
@@ -125,9 +139,10 @@ void FLiveLinkPlaybackTrack::GetFramesUntil(double InPlayhead, TArray<FLiveLinkR
 
 void FLiveLinkPlaybackTrack::GetFramesUntilReverse(double InPlayhead, TArray<FLiveLinkRecordedFrame>& OutFrames)
 {
-	if (LastReadIndex == INDEX_NONE)
+	if (LastReadRelativeIndex == INDEX_NONE)
 	{
-		LastReadIndex = FrameData.Num();
+		LastReadRelativeIndex = FrameData.Num();
+		LastReadAbsoluteIndex = LastReadRelativeIndex + StartIndexOffset;
 	}
 
 	// We need to look up what the last frame would be if this was running forward, and then end on that frame.
@@ -135,20 +150,31 @@ void FLiveLinkPlaybackTrack::GetFramesUntilReverse(double InPlayhead, TArray<FLi
 	// will differ by a frame with a reverse look up. There's probably a better way of handling this.
 	const int32 FinalFrameIndex = PlayheadToFrameIndex(InPlayhead);
 	
-	for (FLiveLinkPlaybackTrackReverseIterator It = FLiveLinkPlaybackTrackReverseIterator(*this, LastReadIndex - 1); It; ++It)
+	for (FLiveLinkPlaybackTrackReverseIterator It = FLiveLinkPlaybackTrackReverseIterator(*this, GetRelativeIndex(LastReadAbsoluteIndex)); It; ++It)
 	{
-		if (FinalFrameIndex == LastReadIndex)
+		const double FrameTimestamp = It.FrameTimestamp();
+		if (FrameTimestamp == LastTimeStamp)
+		{
+			// Generally the first iteration from LastReadAbsoluteIndex will trigger this,
+			// but it's possible if the maximum buffered frames are small (ie 1), then
+			// the LastReadAbsoluteIndex will now point to a different frame.
+			continue;
+		}
+		
+		if (FrameTimestamp < InPlayhead || FinalFrameIndex == LastReadRelativeIndex + StartIndexOffset)
 		{
 			break;
 		}
 
-		LastReadIndex = It.CurrentIndex();
-
+		LastReadRelativeIndex = It.CurrentIndex();
+		LastReadAbsoluteIndex = LastReadRelativeIndex + StartIndexOffset;
+		LastTimeStamp = FrameTimestamp;
+		
 		FLiveLinkRecordedFrame FrameToPlay;
 		FrameToPlay.Data = It.FrameData();
 		FrameToPlay.SubjectKey = SubjectKey;
 		FrameToPlay.LiveLinkRole = LiveLinkRole;
-		FrameToPlay.FrameIndex = LastReadIndex;
+		FrameToPlay.FrameIndex = LastReadAbsoluteIndex;
 
 		OutFrames.Add(MoveTemp(FrameToPlay));
 	}
@@ -156,15 +182,17 @@ void FLiveLinkPlaybackTrack::GetFramesUntilReverse(double InPlayhead, TArray<FLi
 
 bool FLiveLinkPlaybackTrack::TryGetFrame(int32 InIndex, FLiveLinkRecordedFrame& OutFrame)
 {
+	InIndex = GetRelativeIndex(InIndex);
 	if (InIndex >= 0 && InIndex < FrameData.Num())
 	{
-		LastReadIndex = InIndex;
+		LastReadRelativeIndex = InIndex;
+		LastReadAbsoluteIndex = LastReadRelativeIndex + StartIndexOffset;
 		
 		FLiveLinkRecordedFrame FrameToPlay;
-		FrameToPlay.Data = FrameData[InIndex];
+		FrameToPlay.Data = *FrameData[InIndex];
 		FrameToPlay.SubjectKey = SubjectKey;
 		FrameToPlay.LiveLinkRole = LiveLinkRole;
-		FrameToPlay.FrameIndex = LastReadIndex;
+		FrameToPlay.FrameIndex = LastReadAbsoluteIndex;
 
 		OutFrame = MoveTemp(FrameToPlay);
 		return true;
@@ -187,11 +215,12 @@ int32 FLiveLinkPlaybackTrack::PlayheadToFrameIndex(double InPlayhead)
 		CurrentIndex = Idx;
 	}
 
-	return CurrentIndex;
+	return CurrentIndex + StartIndexOffset;
 }
 
 double FLiveLinkPlaybackTrack::FrameIndexToPlayhead(int32 InIndex)
 {
+	InIndex = GetRelativeIndex(InIndex);
 	if (InIndex >= 0 && InIndex < Timestamps.Num())
 	{
 		return Timestamps[InIndex];
@@ -207,8 +236,9 @@ TArray<FLiveLinkRecordedFrame> FLiveLinkPlaybackTracks::FetchNextFrames(double P
 	if (Tracks.Num())
 	{
 		// todo: sort frames by timestamp
-		for (FLiveLinkPlaybackTrack& Track : Tracks)
+		for (TTuple<FLiveLinkSubjectKey, FLiveLinkPlaybackTrack>& TrackKeyVal : Tracks)
 		{
+			FLiveLinkPlaybackTrack& Track = TrackKeyVal.Value;
 			Track.GetFramesUntil(Playhead, NextFrames);
 		}
 	}
@@ -223,8 +253,9 @@ TArray<FLiveLinkRecordedFrame> FLiveLinkPlaybackTracks::FetchPreviousFrames(doub
 	if (Tracks.Num())
 	{
 		// todo: sort frames by timestamp
-		for (FLiveLinkPlaybackTrack& Track : Tracks)
+		for (TTuple<FLiveLinkSubjectKey, FLiveLinkPlaybackTrack>& TrackKeyVal : Tracks)
 		{
+			FLiveLinkPlaybackTrack& Track = TrackKeyVal.Value;
 			Track.GetFramesUntilReverse(Playhead, PreviousFrames);
 		}
 	}
@@ -238,8 +269,9 @@ TArray<FLiveLinkRecordedFrame> FLiveLinkPlaybackTracks::FetchNextFramesAtIndex(i
 
 	if (FrameIndex >= 0)
 	{
-		for (FLiveLinkPlaybackTrack& Track : Tracks)
+		for (TTuple<FLiveLinkSubjectKey, FLiveLinkPlaybackTrack>& TrackKeyVal : Tracks)
 		{
+			FLiveLinkPlaybackTrack& Track = TrackKeyVal.Value;
 			FLiveLinkRecordedFrame Frame;
 			if (Track.TryGetFrame(FrameIndex, Frame))
 			{
@@ -253,8 +285,9 @@ TArray<FLiveLinkRecordedFrame> FLiveLinkPlaybackTracks::FetchNextFramesAtIndex(i
 
 int32 FLiveLinkPlaybackTracks::PlayheadToFrameIndex(double InPlayhead)
 {
-	for (FLiveLinkPlaybackTrack& Track : Tracks)
+	for (TTuple<FLiveLinkSubjectKey, FLiveLinkPlaybackTrack>& TrackKeyVal : Tracks)
 	{
+		FLiveLinkPlaybackTrack& Track = TrackKeyVal.Value;
 		// todo: Is this the best way to determine if this is keyframe data and not static data?
 		if (Track.LiveLinkRole == nullptr)
 		{
@@ -267,8 +300,9 @@ int32 FLiveLinkPlaybackTracks::PlayheadToFrameIndex(double InPlayhead)
 
 double FLiveLinkPlaybackTracks::FrameIndexToPlayhead(int32 InIndex)
 {
-	for (FLiveLinkPlaybackTrack& Track : Tracks)
+	for (TTuple<FLiveLinkSubjectKey, FLiveLinkPlaybackTrack>& TrackKeyVal : Tracks)
 	{
+		FLiveLinkPlaybackTrack& Track = TrackKeyVal.Value;
 		// todo: Is this the best way to determine if this is keyframe data and not static data?
 		if (Track.LiveLinkRole == nullptr)
 		{
@@ -281,20 +315,22 @@ double FLiveLinkPlaybackTracks::FrameIndexToPlayhead(int32 InIndex)
 
 void FLiveLinkPlaybackTracks::Restart(int32 InIndex)
 {
-	for (FLiveLinkPlaybackTrack& Track : Tracks)
+	for (TTuple<FLiveLinkSubjectKey, FLiveLinkPlaybackTrack>& TrackKeyVal : Tracks)
 	{
+		FLiveLinkPlaybackTrack& Track = TrackKeyVal.Value;
 		Track.Restart(InIndex);
 	}
 }
 
 FFrameRate FLiveLinkPlaybackTracks::GetInitialFrameRate() const
 {
-	for (const FLiveLinkPlaybackTrack& Track : Tracks)
+	for (const TTuple<FLiveLinkSubjectKey, FLiveLinkPlaybackTrack>& TrackKeyVal : Tracks)
 	{
+		const FLiveLinkPlaybackTrack& Track = TrackKeyVal.Value;
 		if (Track.LiveLinkRole == nullptr && Track.FrameData.Num() > 0)
 		{
 			FLiveLinkFrameDataStruct FrameDataStruct;
-			FrameDataStruct.InitializeWith(Track.FrameData[0].GetScriptStruct(), (FLiveLinkBaseFrameData*)Track.FrameData[0].GetMemory());
+			FrameDataStruct.InitializeWith(Track.FrameData[0]->GetScriptStruct(), (FLiveLinkBaseFrameData*)Track.FrameData[0]->GetMemory());
 
 			return FrameDataStruct.GetBaseData()->MetaData.SceneTime.Rate;
 		}
@@ -305,30 +341,59 @@ FFrameRate FLiveLinkPlaybackTracks::GetInitialFrameRate() const
 	return FFrameRate(30, 1);
 }
 
-void FLiveLinkUAssetRecordingPlayer::PreparePlayback(const ULiveLinkRecording* CurrentRecording)
+void FLiveLinkUAssetRecordingPlayer::PreparePlayback(ULiveLinkRecording* CurrentRecording)
 {
-	const ULiveLinkUAssetRecording* UAssetRecording = CastChecked<ULiveLinkUAssetRecording>(CurrentRecording);
+	ULiveLinkUAssetRecording* UAssetRecording = CastChecked<ULiveLinkUAssetRecording>(CurrentRecording);
+	LoadedRecording = UAssetRecording;
 
-	FLiveLinkPlaybackTracks RecordingPlayback;
+	// Ensure there nothing is playing and all settings are default. It's possible the CurrentRecording has settings that need to be cleared,
+	// such as if this was just recorded and is now being loaded.
+	ShutdownPlayback();
+	
+	CurrentRecordingPlayback = FLiveLinkPlaybackTracks();
 
-	for (const TPair<FLiveLinkSubjectKey, FLiveLinkRecordingStaticDataContainer>& Pair : UAssetRecording->RecordingData.StaticData)
+	StreamPlayback(0);
+}
+
+void FLiveLinkUAssetRecordingPlayer::ShutdownPlayback()
+{
+	if (LoadedRecording)
 	{
-		FLiveLinkPlaybackTrack PlaybackTrack;
-		PlaybackTrack.FrameData = TConstArrayView<FInstancedStruct>(Pair.Value.RecordedData);
-		PlaybackTrack.Timestamps = TConstArrayView<double>(Pair.Value.Timestamps);
-		PlaybackTrack.LiveLinkRole = Pair.Value.Role;
-		PlaybackTrack.SubjectKey = Pair.Key;
-		RecordingPlayback.Tracks.Add(PlaybackTrack);
+		LoadedRecording->UnloadRecording();
 	}
+}
+
+void FLiveLinkUAssetRecordingPlayer::StreamPlayback(int32 InFromFrame)
+{
+	const int32 InitialFramesToBuffer = GetNumFramesToBuffer();
+	LoadedRecording->LoadRecording(InFromFrame, InitialFramesToBuffer);
+
+	// Make sure there are a few frames ready.
+	LoadedRecording->WaitForBufferedFrames(InFromFrame, InFromFrame + 2);
 
-	for (const TPair<FLiveLinkSubjectKey, FLiveLinkRecordingBaseDataContainer>& Pair : UAssetRecording->RecordingData.FrameData)
+	// On initial load, the correct frame size may not be calculated prior until after waiting for the buffer,
+	// update the correct number of frames and start buffering them.
+	const int32 CurrentFramesToBuffer = GetNumFramesToBuffer();
+	if (CurrentFramesToBuffer != InitialFramesToBuffer)
 	{
-		FLiveLinkPlaybackTrack PlaybackTrack;
-		PlaybackTrack.FrameData = TConstArrayView<FInstancedStruct>(Pair.Value.RecordedData);
-		PlaybackTrack.Timestamps = TConstArrayView<double>(Pair.Value.Timestamps);
-		PlaybackTrack.SubjectKey = Pair.Key;
-		RecordingPlayback.Tracks.Add(PlaybackTrack);
+		LoadedRecording->LoadRecording(InFromFrame, CurrentFramesToBuffer);
 	}
 
-	CurrentRecordingPlayback = MoveTemp(RecordingPlayback);
+	// Take the available recording data.
+	LoadedRecording->CopyRecordingData(CurrentRecordingPlayback);
+}
+
+int32 FLiveLinkUAssetRecordingPlayer::GetNumFramesToBuffer() const
+{
+	const int32 FrameSize = LoadedRecording->GetFrameDiskSize();
+	const int32 MaxFrameBufferSizeMB = GetDefault<ULiveLinkHubSettings>()->FrameBufferSizeMB;
+	const int32 MaxFrameBufferSizeBytes = MaxFrameBufferSizeMB * 1024 * 1024;
+
+	// We divide total frames by 2, since they get doubled later to account for scrubbing in both directions.
+	int32 TotalFramesToBuffer = FrameSize > 0 ? MaxFrameBufferSizeBytes / FrameSize / 2 : 0;
+
+	// Ensure at least a few frames can be buffered.
+	constexpr int32 MinFrames = 3;
+	TotalFramesToBuffer = FMath::Max(TotalFramesToBuffer, MinFrames);
+	return TotalFramesToBuffer;
 }
diff --git a/Engine/Plugins/Animation/LiveLinkHub/Source/LiveLinkHub/Private/Recording/Implementations/LiveLinkUAssetRecordingPlayer.h b/Engine/Plugins/Animation/LiveLinkHub/Source/LiveLinkHub/Private/Recording/Implementations/LiveLinkUAssetRecordingPlayer.h
--- a/Engine/Plugins/Animation/LiveLinkHub/Source/LiveLinkHub/Private/Recording/Implementations/LiveLinkUAssetRecordingPlayer.h
+++ b/Engine/Plugins/Animation/LiveLinkHub/Source/LiveLinkHub/Private/Recording/Implementations/LiveLinkUAssetRecordingPlayer.h
@@ -34,19 +34,36 @@ struct FLiveLinkPlaybackTrack
 	/** Reset the LastReadIndex. */
 	void Restart(int32 NewIndex = INDEX_NONE)
 	{
-		LastReadIndex = NewIndex < FrameData.Num() && NewIndex < Timestamps.Num() ? NewIndex : INDEX_NONE;
+		LastReadRelativeIndex = NewIndex < FrameData.Num() && NewIndex < Timestamps.Num() ? NewIndex : INDEX_NONE;
+		LastReadAbsoluteIndex = LastReadRelativeIndex;
+	}
+
+	/** Convert an absolute frame index to a relative frame index. */
+	int32 GetRelativeIndex(int32 InAbsoluteIndex) const
+	{
+		int32 RelativeIndex = InAbsoluteIndex - StartIndexOffset;
+		RelativeIndex = FMath::Clamp(RelativeIndex, 0, FrameData.Num() - 1);
+		return RelativeIndex;
 	}
 
 	/** Frame data to read. */
-	TConstArrayView<struct FInstancedStruct> FrameData;
+	TConstArrayView<TSharedPtr<FInstancedStruct>> FrameData;
 	/** Timestamps for the frames in the track. */
 	TConstArrayView<double> Timestamps;
 	/** Used for static data. */
 	TSubclassOf<ULiveLinkRole> LiveLinkRole;
 	/** Subject key. */
 	FLiveLinkSubjectKey SubjectKey;
-	/** Index of the last frame that was read by the GetFrames method. */
-	int32 LastReadIndex = -1;
+	/** Index of the last relative frame that was read by the GetFrames method. */
+	int32 LastReadRelativeIndex = -1;
+	/** Index of the last absolute frame that was read by the GetFrames method. */
+	int32 LastReadAbsoluteIndex = -1;
+	/** The true index FrameData starts at. IE, if it starts at 5, then there are 5 prior frames [0..4] that aren't loaded. */
+	int32 StartIndexOffset = 0;
+
+private:
+	/** The last timestamp recorded. */
+	double LastTimeStamp = -1.f;
 
 	friend class FLiveLinkPlaybackTrackIterator;
 };
@@ -76,39 +93,34 @@ struct FLiveLinkPlaybackTracks
 
 public:
 	/** LiveLink tracks to playback. */
-	TArray<FLiveLinkPlaybackTrack> Tracks;
+	TMap<FLiveLinkSubjectKey, FLiveLinkPlaybackTrack> Tracks;
 };
 
 class FLiveLinkUAssetRecordingPlayer : public ILiveLinkRecordingPlayer
 {
 public:
-	void PreparePlayback(const class ULiveLinkRecording* CurrentRecording);
+	virtual void PreparePlayback(class ULiveLinkRecording* CurrentRecording) override;
 
-	virtual TArray<FLiveLinkRecordedFrame> FetchNextFramesAtTimestamp(double Playhead) override
+	virtual void ShutdownPlayback() override;
+	
+	virtual TArray<FLiveLinkRecordedFrame> FetchNextFramesAtTimestamp(const FQualifiedFrameTime& InFrameTime) override
 	{
-		return CurrentRecordingPlayback.FetchNextFrames(Playhead);
+		StreamPlayback(InFrameTime.Time.GetFrame().Value);
+		return CurrentRecordingPlayback.FetchNextFrames(InFrameTime.AsSeconds());
 	}
 
-	virtual TArray<FLiveLinkRecordedFrame> FetchPreviousFramesAtTimestamp(double Playhead) override
+	virtual TArray<FLiveLinkRecordedFrame> FetchPreviousFramesAtTimestamp(const FQualifiedFrameTime& InFrameTime) override
 	{
-		return CurrentRecordingPlayback.FetchPreviousFrames(Playhead);
+		StreamPlayback(InFrameTime.Time.GetFrame().Value);
+		return CurrentRecordingPlayback.FetchPreviousFrames(InFrameTime.AsSeconds());
 	}
 
 	virtual TArray<FLiveLinkRecordedFrame> FetchNextFramesAtIndex(int32 FrameIndex) override
 	{
+		StreamPlayback(FrameIndex);
 		return CurrentRecordingPlayback.FetchNextFramesAtIndex(FrameIndex);
 	}
 
-	virtual int32 PlayheadToFrameIndex(double InPlayhead, bool bReverse) override
-	{
-		return CurrentRecordingPlayback.PlayheadToFrameIndex(InPlayhead);
-	}
-
-	virtual double FrameIndexToPlayhead(int32 InIndex) override
-	{
-		return CurrentRecordingPlayback.FrameIndexToPlayhead(InIndex);
-	}
-
 	virtual void RestartPlayback(int32 InIndex) override
 	{
 		CurrentRecordingPlayback.Restart(InIndex);
@@ -118,8 +130,23 @@ public:
 	{
 		return CurrentRecordingPlayback.GetInitialFrameRate();
 	}
+
+	virtual TRange<int32> GetBufferedFrames() override
+	{
+		return LoadedRecording ? LoadedRecording->GetBufferedFrames() : TRange<int32>(0, 0);
+	}
+
+private:
+	/** Buffer playback around a given frame. */
+	void StreamPlayback(int32 InFromFrame);
+
+	/** Retrieve the total frames to buffer, based on the size the user specified in the config file. */
+	int32 GetNumFramesToBuffer() const;
 	
 private:
 	/** All tracks for the current recording. */
 	FLiveLinkPlaybackTracks CurrentRecordingPlayback;
+
+	/** The recording currently loaded. */
+	TObjectPtr<ULiveLinkUAssetRecording> LoadedRecording;
 };
diff --git a/Engine/Plugins/Animation/LiveLinkHub/Source/LiveLinkHub/Private/Recording/LiveLinkHubPlaybackController.cpp b/Engine/Plugins/Animation/LiveLinkHub/Source/LiveLinkHub/Private/Recording/LiveLinkHubPlaybackController.cpp
--- a/Engine/Plugins/Animation/LiveLinkHub/Source/LiveLinkHub/Private/Recording/LiveLinkHubPlaybackController.cpp
+++ b/Engine/Plugins/Animation/LiveLinkHub/Source/LiveLinkHub/Private/Recording/LiveLinkHubPlaybackController.cpp
@@ -71,7 +71,7 @@ FLiveLinkHubPlaybackController::FLiveLinkHubPlaybackController()
 
 FLiveLinkHubPlaybackController::~FLiveLinkHubPlaybackController()
 {
-	StopPlayback();
+	Eject();
 	Stopping = true;
 	PlaybackEvent->Trigger();
 
@@ -118,6 +118,7 @@ TSharedRef<SWidget> FLiveLinkHubPlaybackController::MakePlaybackWidget()
 		{
 			SliderViewRange = MoveTemp(NewRange);
 		})
+		.GetBufferRange_Raw(this, &FLiveLinkHubPlaybackController::GetBufferedFrames)
 		.GetTotalLength_Raw(this, &FLiveLinkHubPlaybackController::GetLength)
 		.GetCurrentTime_Raw(this, &FLiveLinkHubPlaybackController::GetCurrentTime)
 		.GetSelectionStartTime_Raw(this, &FLiveLinkHubPlaybackController::GetSelectionStartTime)
@@ -314,6 +315,7 @@ void FLiveLinkHubPlaybackController::Eject(TFunction<void()> CompletionCallback)
 	StartTimestamp = 0.f;
 
 	// Recording is done, clear the pointer.
+	RecordingPlayer->ShutdownPlayback();
 	RecordingToPlay.Reset();
 	
 	if (RollbackPreset.IsValid())
@@ -397,6 +399,11 @@ FFrameRate FLiveLinkHubPlaybackController::GetFrameRate() const
 	return CurrentFrameRate;
 }
 
+TRange<int32> FLiveLinkHubPlaybackController::GetBufferedFrames() const
+{
+	return RecordingPlayer->GetBufferedFrames();
+}
+
 void FLiveLinkHubPlaybackController::Start()
 {
 	FString ThreadName = TEXT("LiveLinkHub Playback Controller ");
@@ -518,9 +525,9 @@ void FLiveLinkHubPlaybackController::PushSubjectData(const FLiveLinkRecordedFram
 
 bool FLiveLinkHubPlaybackController::SyncToPlayhead()
 {
-	const double Timestamp = Playhead->GetValue().AsSeconds();
-	TArray<FLiveLinkRecordedFrame> NextFrames = bIsReverse ? RecordingPlayer->FetchPreviousFramesAtTimestamp(Timestamp)
-		: RecordingPlayer->FetchNextFramesAtTimestamp(Timestamp);
+	const FQualifiedFrameTime FrameTime = Playhead->GetValue();
+	TArray<FLiveLinkRecordedFrame> NextFrames = bIsReverse ? RecordingPlayer->FetchPreviousFramesAtTimestamp(FrameTime)
+		: RecordingPlayer->FetchNextFramesAtTimestamp(FrameTime);
 	
 	for (const FLiveLinkRecordedFrame& NextFrame : NextFrames)
 	{
diff --git a/Engine/Plugins/Animation/LiveLinkHub/Source/LiveLinkHub/Private/Recording/LiveLinkHubPlaybackController.h b/Engine/Plugins/Animation/LiveLinkHub/Source/LiveLinkHub/Private/Recording/LiveLinkHubPlaybackController.h
--- a/Engine/Plugins/Animation/LiveLinkHub/Source/LiveLinkHub/Private/Recording/LiveLinkHubPlaybackController.h
+++ b/Engine/Plugins/Animation/LiveLinkHub/Source/LiveLinkHub/Private/Recording/LiveLinkHubPlaybackController.h
@@ -82,6 +82,9 @@ public:
 
 	/** Retrieve the current framerate. */
 	FFrameRate GetFrameRate() const;
+
+	/** Retrieve buffered frame range. */
+	TRange<int32> GetBufferedFrames() const;
 	
 	/** If the controller is ready for commands. */
 	bool IsReady() const
diff --git a/Engine/Plugins/Animation/LiveLinkHub/Source/LiveLinkHub/Private/Recording/LiveLinkRecordingPlayer.h b/Engine/Plugins/Animation/LiveLinkHub/Source/LiveLinkHub/Private/Recording/LiveLinkRecordingPlayer.h
--- a/Engine/Plugins/Animation/LiveLinkHub/Source/LiveLinkHub/Private/Recording/LiveLinkRecordingPlayer.h
+++ b/Engine/Plugins/Animation/LiveLinkHub/Source/LiveLinkHub/Private/Recording/LiveLinkRecordingPlayer.h
@@ -29,26 +29,26 @@ public:
 	virtual ~ILiveLinkRecordingPlayer() = default;
 
 	/** Initialize internal structures needed for playback of the recorded data. */
-	virtual void PreparePlayback(const ULiveLinkRecording* Recording) = 0;
+	virtual void PreparePlayback(ULiveLinkRecording* Recording) = 0;
+
+	/** Preform cleanup when exiting playback. */
+	virtual void ShutdownPlayback() = 0;
 
 	/** Fetch next frames at the provided playhead position. */
-	virtual TArray<FLiveLinkRecordedFrame> FetchNextFramesAtTimestamp(double Playhead) = 0;
+	virtual TArray<FLiveLinkRecordedFrame> FetchNextFramesAtTimestamp(const FQualifiedFrameTime& InFrameTime) = 0;
 
 	/** Fetch previous frames at the provided playhead position. */
-	virtual TArray<FLiveLinkRecordedFrame> FetchPreviousFramesAtTimestamp(double Playhead) = 0;
+	virtual TArray<FLiveLinkRecordedFrame> FetchPreviousFramesAtTimestamp(const FQualifiedFrameTime& InFrameTime) = 0;
 	
 	/** Fetch next frames at the provided frame index. */
 	virtual TArray<FLiveLinkRecordedFrame> FetchNextFramesAtIndex(int32 FrameIndex) = 0;
 
-	/** Convert the playhead to a frame index. */
-	virtual int32 PlayheadToFrameIndex(double InPlayhead, bool bReverse) = 0;
-
-	/** Convert the frame index to a timstamp. */
-	virtual double FrameIndexToPlayhead(int32 InIndex) = 0;
-
 	/** Restart the recording from the beginning. */
 	virtual void RestartPlayback(int32 InIndex = INDEX_NONE) = 0;
 
 	/** Retrieve the first frame's frame rate information. */
 	virtual FFrameRate GetInitialFramerate() = 0;
+
+	/** Retrieve the currently buffered frames. */
+	virtual TRange<int32> GetBufferedFrames() = 0;
 };
diff --git a/Engine/Plugins/Animation/LiveLinkHub/Source/LiveLinkHub/Private/UI/Widgets/SLiveLinkHubPlaybackWidget.cpp b/Engine/Plugins/Animation/LiveLinkHub/Source/LiveLinkHub/Private/UI/Widgets/SLiveLinkHubPlaybackWidget.cpp
--- a/Engine/Plugins/Animation/LiveLinkHub/Source/LiveLinkHub/Private/UI/Widgets/SLiveLinkHubPlaybackWidget.cpp
+++ b/Engine/Plugins/Animation/LiveLinkHub/Source/LiveLinkHub/Private/UI/Widgets/SLiveLinkHubPlaybackWidget.cpp
@@ -2,6 +2,8 @@
 
 #include "SLiveLinkHubPlaybackWidget.h"
 
+#include "SLiveLinkHubTimeSlider.h"
+
 #include "FrameNumberNumericInterface.h"
 #include "Framework/MultiBox/MultiBoxBuilder.h"
 #include "SSimpleTimeSlider.h"
@@ -43,6 +45,8 @@ void SLiveLinkHubPlaybackWidget::Construct(const FArguments& InArgs)
 
 	OnGetFrameRate = InArgs._GetFrameRate;
 
+	OnGetFrameBufferRange = InArgs._GetBufferRange;
+
 	const TAttribute<EFrameNumberDisplayFormats> GetDisplayFormatAttr = MakeAttributeSP(this, &SLiveLinkHubPlaybackWidget::GetDisplayFormat);
 	const TAttribute<FFrameRate> GetDisplayRateAttr = MakeAttributeSP(this, &SLiveLinkHubPlaybackWidget::GetFrameRate);
 	
@@ -57,7 +61,8 @@ void SLiveLinkHubPlaybackWidget::Construct(const FArguments& InArgs)
 		.AutoHeight()
 		.Padding(16.f, 8.f)
 		[
-			SNew(SSimpleTimeSlider)
+			SNew(SLiveLinkHubTimeSlider)
+			.BaseArgs(SSimpleTimeSlider::FArguments()
 			.ClampRangeHighlightSize(0.15f)
 			.ClampRangeHighlightColor(FLinearColor::Gray.CopyWithNewOpacity(0.5f))
 			.ScrubPosition_Lambda([this]()
@@ -78,6 +83,8 @@ void SLiveLinkHubPlaybackWidget::Construct(const FArguments& InArgs)
 								SetCurrentTime(NewScrubTime * GetFrameRate().Numerator);
 							}
 						})
+			)
+			.BufferRange(this, &SLiveLinkHubPlaybackWidget::GetBufferRange)
 		]
 		+SVerticalBox::Slot()
 		.AutoHeight()
@@ -396,6 +403,15 @@ TRange<double> SLiveLinkHubPlaybackWidget::GetClampRange() const
 	return TRange<double>(Start, End);
 }
 
+TRange<double> SLiveLinkHubPlaybackWidget::GetBufferRange() const
+{
+	check(OnGetFrameBufferRange.IsBound());
+	TRange<int32> BufferedFrames = OnGetFrameBufferRange.Execute();
+	FQualifiedFrameTime StartTime(BufferedFrames.GetLowerBound().GetValue(), GetFrameRate());
+	FQualifiedFrameTime EndTime(BufferedFrames.GetUpperBound().GetValue(), GetFrameRate());
+	return TRange<double>(StartTime.AsSeconds(), EndTime.AsSeconds());
+}
+
 bool SLiveLinkHubPlaybackWidget::IsPaused() const
 {
 	return OnGetPausedDelegate.IsBound() ? OnGetPausedDelegate.Execute() : false;
diff --git a/Engine/Plugins/Animation/LiveLinkHub/Source/LiveLinkHub/Private/UI/Widgets/SLiveLinkHubPlaybackWidget.h b/Engine/Plugins/Animation/LiveLinkHub/Source/LiveLinkHub/Private/UI/Widgets/SLiveLinkHubPlaybackWidget.h
--- a/Engine/Plugins/Animation/LiveLinkHub/Source/LiveLinkHub/Private/UI/Widgets/SLiveLinkHubPlaybackWidget.h
+++ b/Engine/Plugins/Animation/LiveLinkHub/Source/LiveLinkHub/Private/UI/Widgets/SLiveLinkHubPlaybackWidget.h
@@ -22,9 +22,10 @@ public:
 	DECLARE_DELEGATE_RetVal(FQualifiedFrameTime, FOnGetTime);
 	DECLARE_DELEGATE_OneParam(FOnSetTime, FQualifiedFrameTime);
 	DECLARE_DELEGATE_RetVal(FFrameRate, FOnGetFrame);
-	DECLARE_DELEGATE_OneParam(FOnSetViewRange, const TRange<double>&)
-	DECLARE_DELEGATE_RetVal(TRange<double>, FOnGetViewRange)
-
+	DECLARE_DELEGATE_OneParam(FOnSetDoubleRange, const TRange<double>&)
+	DECLARE_DELEGATE_RetVal(TRange<double>, FOnGetDoubleRange)
+	DECLARE_DELEGATE_RetVal(TRange<int32>, FOnGetIntRange)
+	
 	DECLARE_DELEGATE(FOnButtonPressed);
 	
 	SLATE_BEGIN_ARGS(SLiveLinkHubPlaybackWidget) { }
@@ -66,10 +67,12 @@ public:
 	SLATE_EVENT(FOnGetFrame, GetFrameRate)
 	
 	/** Get the view range (visible selection range). */
-	SLATE_EVENT(FOnSetViewRange, SetViewRange)
+	SLATE_EVENT(FOnSetDoubleRange, SetViewRange)
 	/** Set the view range (visible selection range). */
-	SLATE_EVENT(FOnGetViewRange, GetViewRange)
-	
+	SLATE_EVENT(FOnGetDoubleRange, GetViewRange)
+	/** Retrieve the frame buffer range. */
+	SLATE_EVENT(FOnGetIntRange, GetBufferRange)
+		
 	SLATE_END_ARGS()
 
 	/**
@@ -125,6 +128,9 @@ private:
 	/** Retrieve the range to clamp playback to (the selection). */
 	TRange<double> GetClampRange() const;
 
+	/** Retrieve the buffered frame range. */
+	TRange<double> GetBufferRange() const;
+	
 	/** Is playback paused? */
 	bool IsPaused() const;
 	/** Is playback in reverse? */
@@ -204,10 +210,13 @@ private:
 	FOnSetTime OnSetCurrentTimeDelegate;
 
 	/** Delegate for getting the view range. */
-	FOnGetViewRange OnGetViewRangeDelegate;
+	FOnGetDoubleRange OnGetViewRangeDelegate;
 	
 	/** Delegate for setting the view range. */
-	FOnSetViewRange OnSetViewRangeDelegate;
+	FOnSetDoubleRange OnSetViewRangeDelegate;
+
+	/** Retrieve the buffered frame range. */
+	FOnGetIntRange OnGetFrameBufferRange;
 
 	/** Delegate for getting the selection start. */
 	FOnGetTime OnGetSelectionStartTimeDelegate;
diff --git a/Engine/Plugins/Animation/LiveLinkHub/Source/LiveLinkHub/Private/UI/Widgets/SLiveLinkHubTimeSlider.cpp b/Engine/Plugins/Animation/LiveLinkHub/Source/LiveLinkHub/Private/UI/Widgets/SLiveLinkHubTimeSlider.cpp
new file mode 100644
--- /dev/null
+++ b/Engine/Plugins/Animation/LiveLinkHub/Source/LiveLinkHub/Private/UI/Widgets/SLiveLinkHubTimeSlider.cpp
@@ -0,0 +1,56 @@
+// Copyright Epic Games, Inc. All Rights Reserved.
+
+#include "SLiveLinkHubTimeSlider.h"
+
+void SLiveLinkHubTimeSlider::Construct(const FArguments& InArgs)
+{
+	BufferRange = (InArgs._BufferRange);
+
+	SSimpleTimeSlider::Construct(InArgs._BaseArgs);
+}
+
+int32 SLiveLinkHubTimeSlider::OnPaint(const FPaintArgs& Args, const FGeometry& AllottedGeometry, const FSlateRect& MyCullingRect, FSlateWindowElementList& OutDrawElements, int32 LayerId, const FWidgetStyle& InWidgetStyle, bool bParentEnabled) const
+{
+	int32 Result = SSimpleTimeSlider::OnPaint(Args, AllottedGeometry, MyCullingRect, OutDrawElements, LayerId, InWidgetStyle, bParentEnabled);
+	OnPaintExtendedSlider(MirrorLabels.Get(), AllottedGeometry, MyCullingRect, OutDrawElements, LayerId, InWidgetStyle, bParentEnabled);
+	return Result;
+}
+
+void SLiveLinkHubTimeSlider::OnPaintExtendedSlider(bool bMirrorLabels, const FGeometry& AllottedGeometry, const FSlateRect& MyCullingRect, FSlateWindowElementList& OutDrawElements, int32 LayerId, const FWidgetStyle& InWidgetStyle, bool bParentEnabled) const
+{
+	const bool bEnabled = bParentEnabled;
+	const ESlateDrawEffect DrawEffects = bEnabled ? ESlateDrawEffect::None : ESlateDrawEffect::DisabledEffect;
+
+	TRange<double> LocalViewRange = ViewRange.Get();
+	const float LocalViewRangeMin = LocalViewRange.GetLowerBoundValue();
+	const float LocalViewRangeMax = LocalViewRange.GetUpperBoundValue();
+	const float LocalSequenceLength = LocalViewRangeMax-LocalViewRangeMin;
+	
+	if (LocalSequenceLength > 0)
+	{
+		FScrubRangeToScreen RangeToScreen(LocalViewRange, AllottedGeometry.GetLocalSize());
+	
+		float LeftBuffer =  RangeToScreen.InputToLocalX(BufferRange.Get().GetLowerBoundValue());
+		float RightBuffer =  RangeToScreen.InputToLocalX(BufferRange.Get().GetUpperBoundValue());
+		float Height = AllottedGeometry.GetLocalSize().Y * ClampRangeHighlightSize.Get();
+	
+		FPaintGeometry RangeGeometry;
+		if (bMirrorLabels)
+		{
+			RangeGeometry = AllottedGeometry.ToPaintGeometry(FVector2f(RightBuffer-LeftBuffer, Height), FSlateLayoutTransform(FVector2f(LeftBuffer, 0)));
+		}
+		else
+		{
+			RangeGeometry = AllottedGeometry.ToPaintGeometry(FVector2f(RightBuffer-LeftBuffer, AllottedGeometry.GetLocalSize().Y), FSlateLayoutTransform(FVector2f(LeftBuffer, AllottedGeometry.GetLocalSize().Y - Height / 3.f)));
+		}
+
+		FSlateDrawElement::MakeBox(
+			OutDrawElements,
+			++LayerId,
+			RangeGeometry,
+			CursorBackground,
+			DrawEffects,
+			FLinearColor::White
+			);
+	}
+}
diff --git a/Engine/Plugins/Animation/LiveLinkHub/Source/LiveLinkHub/Private/UI/Widgets/SLiveLinkHubTimeSlider.h b/Engine/Plugins/Animation/LiveLinkHub/Source/LiveLinkHub/Private/UI/Widgets/SLiveLinkHubTimeSlider.h
new file mode 100644
--- /dev/null
+++ b/Engine/Plugins/Animation/LiveLinkHub/Source/LiveLinkHub/Private/UI/Widgets/SLiveLinkHubTimeSlider.h
@@ -0,0 +1,39 @@
+// Copyright Epic Games, Inc. All Rights Reserved.
+
+#pragma once
+
+#include "SSimpleTimeSlider.h"
+
+/**
+ * Live link hub slider for scrubbing playback and reporting the playhead position.
+ * Additionally, it supports displaying currently buffered frames.
+ */
+class SLiveLinkHubTimeSlider : public SSimpleTimeSlider
+{
+public:
+
+	SLATE_BEGIN_ARGS(SLiveLinkHubTimeSlider)
+	{}
+		/* The buffered frame size. */
+		SLATE_ATTRIBUTE(TRange<double>, BufferRange)
+		/** SimpleTimeSlider base args. */
+		SLATE_ARGUMENT(SSimpleTimeSlider::FArguments, BaseArgs)
+	SLATE_END_ARGS()
+
+	/**
+	 * Construct the widget
+	 * 
+	 * @param InArgs   A declaration from which to construct the widget
+	 */
+	void Construct( const FArguments& InArgs );
+
+protected:
+	virtual int32 OnPaint(const FPaintArgs& Args, const FGeometry& AllottedGeometry, const FSlateRect& MyCullingRect, FSlateWindowElementList& OutDrawElements, int32 LayerId, const FWidgetStyle& InWidgetStyle, bool bParentEnabled) const override;
+
+	/** Paint live link hub specific slider details. */
+	void OnPaintExtendedSlider(bool bMirrorLabels, const FGeometry& AllottedGeometry, const FSlateRect& MyCullingRect, FSlateWindowElementList& OutDrawElements, int32 LayerId, const FWidgetStyle& InWidgetStyle, bool bParentEnabled) const;
+	
+private:
+	/** The frame buffer range to render. */
+	TAttribute<TRange<double>> BufferRange;
+};
\ No newline at end of file
