Add analytics when LiveLinkHub's connection is established

diff --git a/Engine/Plugins/Animation/LiveLinkHub/Source/LiveLinkHubMessaging/Private/LiveLinkHubConnectionManager.h b/Engine/Plugins/Animation/LiveLinkHub/Source/LiveLinkHubMessaging/Private/LiveLinkHubConnectionManager.h
--- a/Engine/Plugins/Animation/LiveLinkHub/Source/LiveLinkHubMessaging/Private/LiveLinkHubConnectionManager.h
+++ b/Engine/Plugins/Animation/LiveLinkHub/Source/LiveLinkHubMessaging/Private/LiveLinkHubConnectionManager.h
@@ -3,6 +3,7 @@
 #pragma once
 
 #include "Engine/World.h"
+#include "EngineAnalytics.h"
 #include "Features/IModularFeatures.h"
 #include "ILiveLinkClient.h"
 #include "ILiveLinkModule.h"
@@ -25,6 +26,20 @@ DEFINE_LOG_CATEGORY_STATIC(LogLiveLinkHubConnectionManager, Log, All);
 
 #if WITH_LIVELINK_DISCOVERY_MANAGER_THREAD
 
+namespace LiveLinkHubConnectionManager
+{
+	static void SendAnalyticsConnectionEstablished()
+	{
+		if (!FEngineAnalytics::IsAvailable())
+		{
+			return;
+		}
+
+		FEngineAnalytics::GetProvider().RecordEvent(TEXT("Usage.LiveLinkHub.ConnectionEstablished"), {});
+	}
+}
+
+
 /** This utitlity is meant to be run on an unreal engine instance to look for livelink hub connections and to automatically create the message bus source for it. */
 class FLiveLinkHubConnectionManager
 {
@@ -138,6 +153,7 @@ private:
 			LastAddedSource = TPair<FMessageAddress, TWeakPtr<ILiveLinkSource>>{ PollResult->Address, LiveLinkSource };
 			ILiveLinkHubMessagingModule& HubMessagingModule = FModuleManager::GetModuleChecked<ILiveLinkHubMessagingModule>("LiveLinkHubMessaging");
 			HubMessagingModule.OnConnectionEstablished().Broadcast(SourceId);
+			LiveLinkHubConnectionManager::SendAnalyticsConnectionEstablished();
 		}
 		else
 		{
