LiveLinkHub - Fix linux build issue.

diff --git a/Engine/Plugins/Animation/LiveLinkHub/Source/LiveLinkHub/Private/Recording/Implementations/LiveLinkUAssetRecording.cpp b/Engine/Plugins/Animation/LiveLinkHub/Source/LiveLinkHub/Private/Recording/Implementations/LiveLinkUAssetRecording.cpp
--- a/Engine/Plugins/Animation/LiveLinkHub/Source/LiveLinkHub/Private/Recording/Implementations/LiveLinkUAssetRecording.cpp
+++ b/Engine/Plugins/Animation/LiveLinkHub/Source/LiveLinkHub/Private/Recording/Implementations/LiveLinkUAssetRecording.cpp
@@ -138,7 +138,7 @@ void ULiveLinkUAssetRecording::UnloadRecording()
 	InitialFrameToStream = 0;
 	TotalFramesToStream = 0;
 	
-	BufferedFrames = TRange<int32>(0, 0);
+	SetBufferedFrames(TRange<int32>(0, 0));
 	
 	for (TTuple<FLiveLinkSubjectKey, FLiveLinkRecordingStaticDataContainer>& StaticData : RecordingData.StaticData)
 	{
@@ -173,7 +173,7 @@ void ULiveLinkUAssetRecording::WaitForBufferedFrames(int32 InMinFrame, int32 InM
 		
 		while (true)
 		{
-			TRange<int32> BufferedFramesLocal = BufferedFrames.load();
+			TRange<int32> BufferedFramesLocal = GetBufferedFrames();
 			if (InTotalFrames > TotalFramesToStream
 			|| BufferedFramesLocal.Contains(InRange) || AsyncStreamTask->IsDone())
 			{
@@ -185,6 +185,12 @@ void ULiveLinkUAssetRecording::WaitForBufferedFrames(int32 InMinFrame, int32 InM
 	}
 }
 
+TRange<int32> ULiveLinkUAssetRecording::GetBufferedFrames() const
+{
+	FScopeLock Lock(&BufferedFrameMutex);
+	return BufferedFrames;
+}
+
 void ULiveLinkUAssetRecording::CopyRecordingData(FLiveLinkPlaybackTracks& InOutLiveLinkPlaybackTracks) const
 {
 	FScopeLock Lock(&DataContainerMutex);
@@ -215,6 +221,12 @@ void ULiveLinkUAssetRecording::CopyRecordingData(FLiveLinkPlaybackTracks& InOutL
 	}
 }
 
+void ULiveLinkUAssetRecording::SetBufferedFrames(const TRange<int32>& InNewRange)
+{
+	FScopeLock Lock(&BufferedFrameMutex);
+	BufferedFrames = InNewRange;
+}
+
 void ULiveLinkUAssetRecording::SaveFrameData(FArchive* InFileWriter, const FLiveLinkSubjectKey& InSubjectKey, FLiveLinkRecordingBaseDataContainer& InBaseDataContainer)
 {
 		// This will crash if it fails -- we don't want to save invalid data.
@@ -286,7 +298,7 @@ void ULiveLinkUAssetRecording::LoadRecordingAsync(int32 InStartFrame, int32 InCu
 	InCurrentFrame = FMath::Clamp(InCurrentFrame, 0, MaxPossibleFrame);
 	const int32 EndFrame = InStartFrame + InNumFramesToLoad - 1;
 
-	if (BufferedFrames.load().Contains(TRange<int32>(InStartFrame, FMath::Min(EndFrame, MaxPossibleFrame > 0 ? MaxPossibleFrame : EndFrame))))
+	if (GetBufferedFrames().Contains(TRange<int32>(InStartFrame, FMath::Min(EndFrame, MaxPossibleFrame > 0 ? MaxPossibleFrame : EndFrame))))
 	{
 		// All frames are already buffered.
 		return;
@@ -528,7 +540,7 @@ void ULiveLinkUAssetRecording::LoadFrameData(FLiveLinkRecordingBaseDataContainer
 		}
 
 		// Record all the frames that have been buffered.
-		BufferedFrames = TRange<int32>(LastLoadedLeftFrame, LastLoadedRightFrame);
+		SetBufferedFrames(TRange<int32>(LastLoadedLeftFrame, LastLoadedRightFrame));
 
 		// Output the streamed data to the data container. This will unload unused frames.
 		{
diff --git a/Engine/Plugins/Animation/LiveLinkHub/Source/LiveLinkHub/Private/Recording/Implementations/LiveLinkUAssetRecording.h b/Engine/Plugins/Animation/LiveLinkHub/Source/LiveLinkHub/Private/Recording/Implementations/LiveLinkUAssetRecording.h
--- a/Engine/Plugins/Animation/LiveLinkHub/Source/LiveLinkHub/Private/Recording/Implementations/LiveLinkUAssetRecording.h
+++ b/Engine/Plugins/Animation/LiveLinkHub/Source/LiveLinkHub/Private/Recording/Implementations/LiveLinkUAssetRecording.h
@@ -152,12 +152,15 @@ public:
 	int32 GetFrameDiskSize() const { return FrameDiskSize; }
 	
 	/** Return the currently buffered frame range. */
-	TRange<int32> GetBufferedFrames() const { return BufferedFrames.load(); }
+	TRange<int32> GetBufferedFrames() const;
 
 	/** Copy the asset's loaded recording data to a format suitable for playback in live link. */
 	void CopyRecordingData(FLiveLinkPlaybackTracks& InOutLiveLinkPlaybackTracks) const;
 	
 private:
+	/** Update the buffered frame range. */
+	void SetBufferedFrames(const TRange<int32>& InNewRange);
+	
 	/** Serialize the number of frames (array size) of the BaseDataContainer to the archive. */
 	void SaveFrameData(FArchive* InFileWriter, const FLiveLinkSubjectKey& InSubjectKey, FLiveLinkRecordingBaseDataContainer& InBaseDataContainer);
 	
@@ -232,8 +235,11 @@ private:
 	std::atomic<int32> FrameDiskSize = 0;
 
 	/** Frames that are loaded and buffered. */
-	std::atomic<TRange<int32>> BufferedFrames = TRange<int32>(0, 0);
+	TRange<int32> BufferedFrames = TRange<int32>(0, 0);
 
+	/** Mutex for accessing the buffered frames. */
+	mutable FCriticalSection BufferedFrameMutex;
+	
 	/** Mutex for accessing the data container from multiple threads. */
 	mutable FCriticalSection DataContainerMutex;
 
