LiveLinkHub - Fix loading recordings with multiple static/frame data. Optimize frame buffering, storing more information on an initial load rather than load it each time.

diff --git a/Engine/Plugins/Animation/LiveLinkHub/Source/LiveLinkHub/Private/Recording/Implementations/LiveLinkUAssetRecording.cpp b/Engine/Plugins/Animation/LiveLinkHub/Source/LiveLinkHub/Private/Recording/Implementations/LiveLinkUAssetRecording.cpp
--- a/Engine/Plugins/Animation/LiveLinkHub/Source/LiveLinkHub/Private/Recording/Implementations/LiveLinkUAssetRecording.cpp
+++ b/Engine/Plugins/Animation/LiveLinkHub/Source/LiveLinkHub/Private/Recording/Implementations/LiveLinkUAssetRecording.cpp
@@ -78,8 +78,8 @@ void ULiveLinkUAssetRecording::LoadRecording(int32 InInitialFrame, int32 InNumFr
 	if (!AsyncStreamTask.IsValid())
 	{
 		check(RecordingFileReader == nullptr);
-		
-		RecordingStartFrameFilePosition = 0;
+
+		FrameFileData.Empty();
 		
 		const FString FilePath = GetRecordingDataFilePath();
 		RecordingFileReader = IFileManager::Get().CreateFileReader(*FilePath);
@@ -130,16 +130,13 @@ void ULiveLinkUAssetRecording::UnloadRecording()
 		RecordingFileReader = nullptr;
 	}
 
-	FrameDataSubjectKey.Reset();
-	RecordingStartFrameFilePosition = 0;
+	FrameFileData.Empty();
 	RecordingMaxFrames = 0;
-	FrameDiskSize = 0;
+	MaxFrameDiskSize = 0;
 	EarliestFrameToStream = 0;
 	InitialFrameToStream = 0;
 	TotalFramesToStream = 0;
 	
-	SetBufferedFrames(TRange<int32>(0, 0));
-	
 	for (TTuple<FLiveLinkSubjectKey, FLiveLinkRecordingStaticDataContainer>& StaticData : RecordingData.StaticData)
 	{
 		StaticData.Value.Timestamps.Empty();
@@ -188,7 +185,34 @@ void ULiveLinkUAssetRecording::WaitForBufferedFrames(int32 InMinFrame, int32 InM
 TRange<int32> ULiveLinkUAssetRecording::GetBufferedFrames() const
 {
 	FScopeLock Lock(&BufferedFrameMutex);
-	return BufferedFrames;
+	if (FrameFileData.Num() > 0)
+	{
+		int32 MaxStart = MIN_int32;
+		int32 MinEnd = MAX_int32;
+
+		for (const FFrameFileData& FrameData : FrameFileData)
+		{
+			const int32 RangeStart = FrameData.BufferedFrames.GetLowerBoundValue();
+			const int32 RangeEnd = FrameData.BufferedFrames.GetUpperBoundValue();
+
+			if (RangeStart > MaxStart)
+			{
+				MaxStart = RangeStart;
+			}
+
+			if (RangeEnd < MinEnd)
+			{
+				MinEnd = RangeEnd;
+			}
+		}
+
+		if (MaxStart <= MinEnd)
+		{
+			return TRange<int32>(MaxStart, MinEnd);
+		}
+	}
+	
+	return TRange<int32>(0, 0);
 }
 
 void ULiveLinkUAssetRecording::CopyRecordingData(FLiveLinkPlaybackTracks& InOutLiveLinkPlaybackTracks) const
@@ -221,10 +245,10 @@ void ULiveLinkUAssetRecording::CopyRecordingData(FLiveLinkPlaybackTracks& InOutL
 	}
 }
 
-void ULiveLinkUAssetRecording::SetBufferedFrames(const TRange<int32>& InNewRange)
+void ULiveLinkUAssetRecording::SetBufferedFrames(FFrameFileData& InFrameData, const TRange<int32>& InNewRange)
 {
 	FScopeLock Lock(&BufferedFrameMutex);
-	BufferedFrames = InNewRange;
+	InFrameData.BufferedFrames = InNewRange;
 }
 
 void ULiveLinkUAssetRecording::SaveFrameData(FArchive* InFileWriter, const FLiveLinkSubjectKey& InSubjectKey, FLiveLinkRecordingBaseDataContainer& InBaseDataContainer)
@@ -291,28 +315,43 @@ void ULiveLinkUAssetRecording::SaveFrameData(FArchive* InFileWriter, const FLive
 		}
 }
 
+static TAutoConsoleVariable<float> CVarLiveLinkHubDebugFrameBufferDelay(
+	TEXT("LiveLinkHub.Debug.FrameBufferDelay"),
+	0.0f,
+	TEXT("The number of seconds to wait when buffering each frame.")
+);
+
 void ULiveLinkUAssetRecording::LoadRecordingAsync(int32 InStartFrame, int32 InCurrentFrame, int32 InNumFramesToLoad)
 {
 	const int32 MaxPossibleFrame = RecordingMaxFrames - 1;
 	InStartFrame = FMath::Clamp(InStartFrame, 0, MaxPossibleFrame);
 	InCurrentFrame = FMath::Clamp(InCurrentFrame, 0, MaxPossibleFrame);
 	const int32 EndFrame = InStartFrame + InNumFramesToLoad - 1;
-
+	
+	ON_SCOPE_EXIT
+	{
+		// Always set to true. Some blocking operations wait for this, and in the case of a non-fatal error
+		// we want to display error logs and don't want the program to freeze.
+		bPerformedInitialLoad = true;
+	};
+	
 	if (GetBufferedFrames().Contains(TRange<int32>(InStartFrame, FMath::Min(EndFrame, MaxPossibleFrame > 0 ? MaxPossibleFrame : EndFrame))))
 	{
 		// All frames are already buffered.
 		return;
 	}
-	
-	RecordingFileReader->Seek(RecordingStartFrameFilePosition);
 
+	DebugSleepTime = CVarLiveLinkHubDebugFrameBufferDelay.GetValueOnAnyThread();
+	
 	// Read file for streaming into memory.
 	if (!RecordingFileReader->AtEnd())
 	{
 		// Perform initial load and record entry frame file offsets.
-		const bool bInitialLoad = RecordingStartFrameFilePosition == 0;
+		const bool bInitialLoad = FrameFileData.Num() == 0;
 		if (bInitialLoad)
 		{
+			RecordingFileReader->Seek(0);
+		
 			int32 LoadedRecordingVersion;
 			*RecordingFileReader << LoadedRecordingVersion;
 
@@ -324,18 +363,24 @@ void ULiveLinkUAssetRecording::LoadRecordingAsync(int32 InStartFrame, int32 InCu
 			int32 NumStaticData = 0;
 			*RecordingFileReader << NumStaticData;
 
-			if (NumStaticData > 0)
+			for (int32 StaticIdx = 0; StaticIdx < NumStaticData; ++StaticIdx)
 			{
 				FGuid KeySource = FGuid();
 				FString KeyName;
 			
 				*RecordingFileReader << KeySource;
 				*RecordingFileReader << KeyName;
-			
-				FLiveLinkSubjectKey SubjectKey(KeySource, *KeyName);
 
-				FLiveLinkRecordingStaticDataContainer& DataContainer = RecordingData.StaticData.FindChecked(SubjectKey);
-				LoadFrameData(&DataContainer, 0, 0, 1);
+				// Create framedata just to load initial static frame data. Static data doesn't require this afterward.
+				FFrameFileData TemporaryFrameData;
+				TemporaryFrameData.FrameDataSubjectKey = MakeShared<FLiveLinkSubjectKey>(KeySource, *KeyName);
+				if (!LoadInitialFrameData(TemporaryFrameData))
+				{
+					return;
+				}
+				
+				FLiveLinkRecordingStaticDataContainer& DataContainer = RecordingData.StaticData.FindChecked(*TemporaryFrameData.FrameDataSubjectKey.Get());
+				LoadFrameData(TemporaryFrameData, DataContainer, 0, 0, 1);
 			}
 
 			// Process frame data.
@@ -343,38 +388,46 @@ void ULiveLinkUAssetRecording::LoadRecordingAsync(int32 InStartFrame, int32 InCu
 			int32 NumFrameData = 0;
 			*RecordingFileReader << NumFrameData;
 
-			if (NumFrameData > 0)
+			for (int32 FrameIdx = 0; FrameIdx < NumFrameData; ++FrameIdx)
 			{
 				FGuid KeySource = FGuid();
 				FString KeyName;
 			
 				*RecordingFileReader << KeySource;
 				*RecordingFileReader << KeyName;
-			
-				FrameDataSubjectKey = MakeShared<FLiveLinkSubjectKey>(KeySource, *KeyName);
-			}
 
-			RecordingStartFrameFilePosition = RecordingFileReader->Tell();
+				FFrameFileData KeyPosition;
+				KeyPosition.FrameDataSubjectKey = MakeShared<FLiveLinkSubjectKey>(KeySource, *KeyName);
+				if (!LoadInitialFrameData(KeyPosition))
+				{
+					return;
+				}
+
+				// Offset to the end of this block if there is multiple NumFrameData.
+				const int64 EndBlockPosition = KeyPosition.GetFrameFilePosition(KeyPosition.MaxFrames);
+				RecordingFileReader->Seek(EndBlockPosition);
+				FrameFileData.Add(MoveTemp(KeyPosition));
+			}
 		}
 
 		// Load the required frames, either on initial load or subsequent loads.
-		if (ensure(FrameDataSubjectKey.IsValid()))
+		for (FFrameFileData& FrameData : FrameFileData)
 		{
-			FLiveLinkRecordingBaseDataContainer& DataContainer = RecordingData.FrameData.FindChecked(*FrameDataSubjectKey.Get());
-			LoadFrameData(&DataContainer, InStartFrame, InCurrentFrame, InNumFramesToLoad);
-		}
-		else
-		{
-			UE_LOG(LogLiveLinkHub, Error, TEXT("FrameDataSubjectKey is missing for file %s."), *GetRecordingDataFilePath());
+			if (ensure(FrameData.FrameDataSubjectKey.IsValid()))
+			{
+				FLiveLinkRecordingBaseDataContainer& DataContainer = RecordingData.FrameData.FindChecked(*FrameData.FrameDataSubjectKey.Get());
+				LoadFrameData(FrameData, DataContainer, InStartFrame, InCurrentFrame, InNumFramesToLoad);
+			}
+			else
+			{
+				UE_LOG(LogLiveLinkHub, Error, TEXT("FrameDataSubjectKey is missing for file %s."), *GetRecordingDataFilePath());
+			}
 		}
-		bPerformedInitialLoad = true;
 	}
 }
 
-void ULiveLinkUAssetRecording::LoadFrameData(FLiveLinkRecordingBaseDataContainer* DataContainer, int32 RequestedStartFrame, int32 RequestedInitialFrame, int32 RequestedFramesToLoad)
+bool ULiveLinkUAssetRecording::LoadInitialFrameData(FFrameFileData& OutFrameData)
 {
-	bStreamingFrameChange = false;
-
 	int32 MaxFrames = 0;
 	(*RecordingFileReader) << MaxFrames;
 
@@ -382,48 +435,70 @@ void ULiveLinkUAssetRecording::LoadFrameData(FLiveLinkRecordingBaseDataContainer
 	{
 		RecordingMaxFrames = MaxFrames;
 	}
+	
+	OutFrameData.MaxFrames = MaxFrames;
 			
 	if (MaxFrames > 0)
-	{
-		if (RequestedFramesToLoad > 0)
-		{
-			// Don't go past requested frames or max frames.
-			MaxFrames = FMath::Min(MaxFrames, RequestedStartFrame + RequestedFramesToLoad);
-		}
-				
+	{	
 		FString StructTypeName;
 		int32 SerializedStructureSize;
 			
 		*RecordingFileReader << StructTypeName;
 		*RecordingFileReader << SerializedStructureSize;
 
-		int64 StartPosition = RecordingFileReader->Tell();
-				
-		UScriptStruct* LoadedStruct = FindObject<UScriptStruct>(nullptr, *StructTypeName, true);
-		if (LoadedStruct == nullptr)
+		OutFrameData.SerializedStructureSize = SerializedStructureSize;
+		OutFrameData.RecordingStartFrameFilePosition = RecordingFileReader->Tell();
+
+		OutFrameData.LoadedStruct = FindObject<UScriptStruct>(nullptr, *StructTypeName, true);
+		if (!OutFrameData.LoadedStruct.IsValid())
 		{
 			UE_LOG(LogLiveLinkHub, Error, TEXT("Script struct type '%s' not found."), *StructTypeName);
-			return;
+			return false;
 		}
 
 		// The size on disk for each frame-- consisting of the frame index, timestamp, and frame struct data.
-		FrameDiskSize = (sizeof(int32) + sizeof(double) + SerializedStructureSize);
+		OutFrameData.FrameDiskSize = (sizeof(int32) + sizeof(double) + SerializedStructureSize);
+
+		if (OutFrameData.FrameDiskSize > MaxFrameDiskSize)
+		{
+			MaxFrameDiskSize = OutFrameData.FrameDiskSize;
+		}
+	}
+
+	return true;
+}
+
+void ULiveLinkUAssetRecording::LoadFrameData(FFrameFileData& InFrameData, FLiveLinkRecordingBaseDataContainer& InDataContainer,
+                                             int32 RequestedStartFrame, int32 RequestedInitialFrame, int32 RequestedFramesToLoad)
+{
+	bStreamingFrameChange = false;
+
+	int32 MaxFrames = InFrameData.MaxFrames;
+	if (MaxFrames > 0)
+	{
+		if (RequestedFramesToLoad > 0)
+		{
+			// Don't go past requested frames or max frames.
+			MaxFrames = FMath::Min(MaxFrames, RequestedStartFrame + RequestedFramesToLoad);
+		}
 		
-		auto GetFrameFilePosition = [StartPosition, this](int32 FrameIdx)
+		UScriptStruct* LoadedStruct = InFrameData.LoadedStruct.Get();
+		if (LoadedStruct == nullptr)
 		{
-			return StartPosition + (FrameDiskSize * FrameIdx);
-		};
+			UE_LOG(LogLiveLinkHub, Error, TEXT("Script struct type not found."));
+			return;
+		}
 
 		// Seek to the requested start frame position in the file.
-		int64 NewPosition = GetFrameFilePosition(RequestedStartFrame);
+		int64 NewPosition = InFrameData.GetFrameFilePosition(RequestedStartFrame);
 		RecordingFileReader->Seek(NewPosition);
 
 		// Arrays to store the newly loaded data.
 		TArray<double> NewTimestamps;
 		TArray<TSharedPtr<FInstancedStruct>> NewRecordedData;
 
-		NewTimestamps.Reserve(RequestedFramesToLoad);
-		NewRecordedData.Reserve(RequestedFramesToLoad);
+		NewTimestamps.Reserve(MaxFrames);
+		NewRecordedData.Reserve(MaxFrames);
 
 		// Load each frame from the initial frame, alternating right to left each frame. This creates a buffer to support
 		// scrubbing each direction and makes sure the immediate frames are loaded first.
@@ -467,7 +542,7 @@ void ULiveLinkUAssetRecording::LoadFrameData(FLiveLinkRecordingBaseDataContainer
 				LastLoadedLeftFrame = FrameToLoad;
 			}
 
-			int64 FramePosition = GetFrameFilePosition(FrameToLoad);
+			int64 FramePosition = InFrameData.GetFrameFilePosition(FrameToLoad);
 			RecordingFileReader->Seek(FramePosition);
 
 			int32 ParsedFrameIdx = 0;
@@ -509,9 +584,9 @@ void ULiveLinkUAssetRecording::LoadFrameData(FLiveLinkRecordingBaseDataContainer
 			
 			// Don't load a frame already in memory.
 			double ExistingTimestamp = 0.f;
-			if (TSharedPtr<FInstancedStruct> ExistingFrame = DataContainer->TryGetFrame(ParsedFrameIdx, ExistingTimestamp))
+			if (TSharedPtr<FInstancedStruct> ExistingFrame = InDataContainer.TryGetFrame(ParsedFrameIdx, ExistingTimestamp))
 			{
-				check(RecordingStartFrameFilePosition != 0);
+				check(RecordingFileReader->Tell() != 0);
 				constexpr bool bCopy = true; // Copy, as the data container could still be having its frame data pushed to animation.
 				InsertFrame(ExistingFrame, ExistingTimestamp, bCopy);
 				AlternateLoadDirection();
@@ -531,6 +606,12 @@ void ULiveLinkUAssetRecording::LoadFrameData(FLiveLinkRecordingBaseDataContainer
 			ensure(NewTimestamps.Num() == NewRecordedData.Num());
 
 			AlternateLoadDirection();
+
+			// Test slow buffer.
+			if (DebugSleepTime > 0.f)
+			{
+				FPlatformProcess::Sleep(DebugSleepTime);
+			}
 			
 			if (bStreamingFrameChange)
 			{
@@ -540,14 +621,14 @@ void ULiveLinkUAssetRecording::LoadFrameData(FLiveLinkRecordingBaseDataContainer
 		}
 
 		// Record all the frames that have been buffered.
-		SetBufferedFrames(TRange<int32>(LastLoadedLeftFrame, LastLoadedRightFrame));
+		SetBufferedFrames(InFrameData, TRange<int32>(LastLoadedLeftFrame, LastLoadedRightFrame));
 
 		// Output the streamed data to the data container. This will unload unused frames.
 		{
 			FScopeLock Lock(&DataContainerMutex);
-			DataContainer->Timestamps = MoveTemp(NewTimestamps);
-			DataContainer->RecordedData = MoveTemp(NewRecordedData);
-			DataContainer->RecordedDataStartFrame = LastLoadedLeftFrame;
+			InDataContainer.Timestamps = MoveTemp(NewTimestamps);
+			InDataContainer.RecordedData = MoveTemp(NewRecordedData);
+			InDataContainer.RecordedDataStartFrame = LastLoadedLeftFrame;
 
 			// This could potentially be optimized, such as outputting directly to the data container during the iteration, which could
 			// allow smoother streaming when scrubbing to a position that isn't buffered at all. However, we would need to be careful of the cost
diff --git a/Engine/Plugins/Animation/LiveLinkHub/Source/LiveLinkHub/Private/Recording/Implementations/LiveLinkUAssetRecording.h b/Engine/Plugins/Animation/LiveLinkHub/Source/LiveLinkHub/Private/Recording/Implementations/LiveLinkUAssetRecording.h
--- a/Engine/Plugins/Animation/LiveLinkHub/Source/LiveLinkHub/Private/Recording/Implementations/LiveLinkUAssetRecording.h
+++ b/Engine/Plugins/Animation/LiveLinkHub/Source/LiveLinkHub/Private/Recording/Implementations/LiveLinkUAssetRecording.h
@@ -149,26 +149,54 @@ public:
 	int32 GetMaxFrames() const { return RecordingMaxFrames; }
 
 	/** The size in bytes of each animation frame. */
-	int32 GetFrameDiskSize() const { return FrameDiskSize; }
+	int32 GetFrameDiskSize() const { return MaxFrameDiskSize; }
 	
-	/** Return the currently buffered frame range. */
+	/** Return the minimum buffered frame range. */
 	TRange<int32> GetBufferedFrames() const;
 
 	/** Copy the asset's loaded recording data to a format suitable for playback in live link. */
 	void CopyRecordingData(FLiveLinkPlaybackTracks& InOutLiveLinkPlaybackTracks) const;
 	
 private:
+	/** Frame data file information when loading from a recording file. */
+	struct FFrameFileData
+	{
+		/** The subject key used for the frame data. */
+		TSharedPtr<FLiveLinkSubjectKey> FrameDataSubjectKey;
+		/** The struct for this frame data. */
+		TWeakObjectPtr<UScriptStruct> LoadedStruct;
+		/** The position in the file recording where frame data begins. */
+		int64 RecordingStartFrameFilePosition = 0;
+		/** Maximum number of frames. */
+		int32 MaxFrames = 0;
+		/** Total size of the structure. */
+		int32 SerializedStructureSize = 0;
+		/** The size in bytes of each animation frame. */
+		int32 FrameDiskSize;
+		/** Buffered frames for this framedata. */
+		TRange<int32> BufferedFrames = TRange<int32>(0, 0);
+
+		/** Find the correct file offset based on the frame index. */
+		int64 GetFrameFilePosition(const int32 InFrameIdx) const
+		{
+			return RecordingStartFrameFilePosition + (FrameDiskSize * InFrameIdx);
+		};
+	};
+	
 	/** Update the buffered frame range. */
-	void SetBufferedFrames(const TRange<int32>& InNewRange);
+	void SetBufferedFrames(FFrameFileData& InFrameData, const TRange<int32>& InNewRange);
 	
 	/** Serialize the number of frames (array size) of the BaseDataContainer to the archive. */
 	void SaveFrameData(FArchive* InFileWriter, const FLiveLinkSubjectKey& InSubjectKey, FLiveLinkRecordingBaseDataContainer& InBaseDataContainer);
 	
 	/** Initialize or update an async load. */
 	void LoadRecordingAsync(int32 InStartFrame, int32 InCurrentFrame, int32 InNumFramesToLoad);
+
+	/** Initial processing on a frame, finding the correct struct and offsets. The RecordingFileReader is assumed to be at the correct position.  */
+	bool LoadInitialFrameData(FFrameFileData& OutFrameData);
 	
 	/** Load frame data to a data container. */
-	void LoadFrameData(FLiveLinkRecordingBaseDataContainer* DataContainer,
+	void LoadFrameData(FFrameFileData& InFrameData, FLiveLinkRecordingBaseDataContainer& InDataContainer,
 		int32 RequestedStartFrame, int32 RequestedInitialFrame, int32 RequestedFramesToLoad);
 	
 	/** Retrieve the recording data file path for this asset. */
@@ -204,12 +232,9 @@ private:
 	/** The file reader for the recording data. */
 	FArchive* RecordingFileReader = nullptr;
 
-	/** The subject key used for the frame data. */
-	TSharedPtr<FLiveLinkSubjectKey> FrameDataSubjectKey;
-	
-	/** The position in the file recording where frame data begins. */
-	int64 RecordingStartFrameFilePosition = 0;
-	
+	/** The loaded frame data keys and position. */
+	TArray<FFrameFileData> FrameFileData;
+
 	/** The maximum frames for this recording. */
 	int32 RecordingMaxFrames = 0;
 	
@@ -230,12 +255,9 @@ private:
 
 	/** True once a full initial load has been performed -- static + frame data. */
 	std::atomic<bool> bPerformedInitialLoad = false;
-	
-	/** The size in bytes of each animation frame. */
-	std::atomic<int32> FrameDiskSize = 0;
 
-	/** Frames that are loaded and buffered. */
-	TRange<int32> BufferedFrames = TRange<int32>(0, 0);
+	/** The maximum frame disk size across frame data. */
+	std::atomic<int32> MaxFrameDiskSize = 0;
 
 	/** Mutex for accessing the buffered frames. */
 	mutable FCriticalSection BufferedFrameMutex;
@@ -246,6 +268,9 @@ private:
 	/** The thread streaming data from disk. */
 	TUniquePtr<FAsyncTask<FLiveLinkStreamAsyncTask>> AsyncStreamTask;
 
+	/** Test slow frame buffering. */
+	float DebugSleepTime = 0.f;
+	
 	/** The current version of the recording. */
 	const int32 RecordingVersion = 1;
 };
