LiveLinkHub - Convert LiveLinkHubClient to use LiveLinkClient's base methods

This change allows us to use the LiveLinkClient directly rather than copying the CreateSource method. The playback source factory will provide a custom ILiveLinkSource that acts as a stub.

diff --git a/Engine/Plugins/Animation/LiveLinkHub/Source/LiveLinkHub/Private/LiveLinkHub.cpp b/Engine/Plugins/Animation/LiveLinkHub/Source/LiveLinkHub/Private/LiveLinkHub.cpp
--- a/Engine/Plugins/Animation/LiveLinkHub/Source/LiveLinkHub/Private/LiveLinkHub.cpp
+++ b/Engine/Plugins/Animation/LiveLinkHub/Source/LiveLinkHub/Private/LiveLinkHub.cpp
@@ -159,7 +159,7 @@ void FLiveLinkHub::OnStaticDataReceived_AnyThread(const FLiveLinkSubjectKey& InS
 	LiveLinkProvider->UpdateSubjectStaticData(OverridenName, InRole, MoveTemp(StaticDataCopy));
 }
 
-void FLiveLinkHub::OnFrameDataReceived_AnyThread(const FLiveLinkSubjectKey& InSubjectKey, FLiveLinkFrameDataStruct& InFrameDataStruct) const
+void FLiveLinkHub::OnFrameDataReceived_AnyThread(const FLiveLinkSubjectKey& InSubjectKey, const FLiveLinkFrameDataStruct& InFrameDataStruct) const
 {
 	if (RecordingController->IsRecording())
 	{
diff --git a/Engine/Plugins/Animation/LiveLinkHub/Source/LiveLinkHub/Private/LiveLinkHub.h b/Engine/Plugins/Animation/LiveLinkHub/Source/LiveLinkHub/Private/LiveLinkHub.h
--- a/Engine/Plugins/Animation/LiveLinkHub/Source/LiveLinkHub/Private/LiveLinkHub.h
+++ b/Engine/Plugins/Animation/LiveLinkHub/Source/LiveLinkHub/Private/LiveLinkHub.h
@@ -73,7 +73,7 @@ public:
 private:
 	//~ LiveLink Client delegates
 	void OnStaticDataReceived_AnyThread(const FLiveLinkSubjectKey& InSubjectKey, TSubclassOf<ULiveLinkRole> InRole, const FLiveLinkStaticDataStruct& InStaticDataStruct) const;
-	void OnFrameDataReceived_AnyThread(const FLiveLinkSubjectKey& InSubjectKey, FLiveLinkFrameDataStruct& InFrameDataStruct) const;
+	void OnFrameDataReceived_AnyThread(const FLiveLinkSubjectKey& InSubjectKey, const FLiveLinkFrameDataStruct& InFrameDataStruct) const;
 	void OnSubjectMarkedPendingKill_AnyThread(const FLiveLinkSubjectKey& InSubjectKey) const;
 	//~ LiveLink Client delegates
 
diff --git a/Engine/Plugins/Animation/LiveLinkHub/Source/LiveLinkHub/Private/LiveLinkHubClient.cpp b/Engine/Plugins/Animation/LiveLinkHub/Source/LiveLinkHub/Private/LiveLinkHubClient.cpp
--- a/Engine/Plugins/Animation/LiveLinkHub/Source/LiveLinkHub/Private/LiveLinkHubClient.cpp
+++ b/Engine/Plugins/Animation/LiveLinkHub/Source/LiveLinkHub/Private/LiveLinkHubClient.cpp
@@ -2,27 +2,18 @@
 
 #include "LiveLinkHubClient.h"
 
-#include "Algo/AnyOf.h"
-#include "Async/TaskGraphInterfaces.h"
 #include "Clients/LiveLinkHubProvider.h"
-#include "Engine/Engine.h"
-#include "LiveLinkClient.h"
 #include "LiveLinkHub.h"
-#include "LiveLinkHubLog.h"
 #include "LiveLinkHubModule.h"
-#include "LiveLinkHubSubjectSettings.h"
 #include "LiveLinkRoleTrait.h"
 #include "LiveLinkSourceCollection.h"
-#include "LiveLinkSettings.h"
 #include "LiveLinkSubject.h"
-#include "LiveLinkSubjectSettings.h"
-#include "LiveLinkTimedDataInput.h"
 #include "Misc/App.h"
 #include "Modules/ModuleManager.h"
-#include "Recording/LiveLinkPlaybackSource.h"
-#include "Recording/LiveLinkPlaybackSubject.h"
+#include "Recording/LiveLinkHubPlaybackSourceSettings.h"
 #include "Stats/Stats2.h"
 #include "UObject/Package.h"
+#include "UObject/StrongObjectPtr.h"
 #include "UObject/UObjectGlobals.h"
 
 DECLARE_STATS_GROUP(TEXT("Live Link Hub"), STATGROUP_LiveLinkHub, STATCAT_Advanced);
@@ -60,220 +51,22 @@ bool FLiveLinkHubClient::CreateSource(const FLiveLinkSourcePreset& InSourcePrese
 	// Create fake sources if we're in playback mode so we don't process livelink data while we're playing a recording.
 	if (LiveLinkHub.Pin()->IsInPlayback())
 	{
-		return CreatePlaybackSource(InSourcePreset);
-	}
-	else
-	{
-		return FLiveLinkClient::CreateSource(InSourcePreset);
-	}
-}
-
-bool FLiveLinkHubClient::CreatePlaybackSource(const FLiveLinkSourcePreset& InSourcePreset)
-{
-	check(Collection);
-
-	if (InSourcePreset.Settings == nullptr)
-	{
-		UE_LOG(LogLiveLinkHub, Warning, TEXT("Create Source Failure: The settings are not defined."));
-		return false;
-	}
-
-	if (InSourcePreset.Guid == FLiveLinkSourceCollection::DefaultVirtualSubjectGuid)
-	{
-		UE_LOG(LogLiveLinkHub, Warning, TEXT("Create Source Failure: Can't create default subject source. It will be created automatically."));
-		return false;
-	}
-
-	if (!InSourcePreset.Guid.IsValid())
-	{
-		UE_LOG(LogLiveLinkHub, Warning, TEXT("Create Source Failure: The guid is invalid."));
-		return false;
-	}
-
-	if (Collection->FindSource(InSourcePreset.Guid) != nullptr)
-	{
-		UE_LOG(LogLiveLinkHub, Warning, TEXT("Create Source Failure: The guid already exist."));
-		return false;
-	}
-
-	ULiveLinkSourceSettings* Setting = nullptr;
-	TSharedPtr<ILiveLinkSource> Source;
-	FLiveLinkCollectionSourceItem Data;
-	Data.Guid = InSourcePreset.Guid;
-
-	// subject source have a special settings class. We can differentiate them using this
-	if (InSourcePreset.Settings->GetClass()->GetName() == TEXT("LiveLinkVirtualSubjectSourceSettings"))
-	{
-		Source = MakeShared<FLiveLinkPlaybackSource>();
-		Data.bIsVirtualSource = true;
-	}
-	else
-	{
-		Source = MakeShared<FLiveLinkPlaybackSource>();
-		Data.TimedData = MakeShared<FLiveLinkTimedDataInput>(this, InSourcePreset.Guid);
-	}
-
-	Data.Source = Source;
-
-	// In case a source has changed its source settings class, instead of duplicating, create the right one and copy previous properties
-	UClass* SourceSettingsClass = Source->GetSettingsClass().Get();
-	if (SourceSettingsClass && SourceSettingsClass != InSourcePreset.Settings->GetClass())
-	{
-		Setting = NewObject<ULiveLinkSourceSettings>(GetTransientPackage(), SourceSettingsClass);
-		UEngine::CopyPropertiesForUnrelatedObjects(InSourcePreset.Settings, Setting);
-		Data.Setting = TStrongObjectPtr(Setting);
-	}
-	else
-	{
-		Data.Setting = TStrongObjectPtr(DuplicateObject<ULiveLinkSourceSettings>(InSourcePreset.Settings, GetTransientPackage()));
-		Setting = Data.Setting.Get();
-	}
-
-	Collection->AddSource(MoveTemp(Data));
-
-	Source->ReceiveClient(this, InSourcePreset.Guid);
-	Source->InitializeSettings(Setting);
-
-	return true;
-}
-
-bool FLiveLinkHubClient::CreateSubject(const FLiveLinkSubjectPreset& InSubjectPreset)
-{
-	// Create fake subjects if we're in playback mode
-	if (LiveLinkHub.Pin()->IsInPlayback())
-	{
-		return CreatePlaybackSubject(InSubjectPreset);
-	}
-	else
-	{
-		return FLiveLinkClient::CreateSubject(InSubjectPreset);
-	}
-}
-
-bool FLiveLinkHubClient::CreatePlaybackSubject(const FLiveLinkSubjectPreset& InSubjectPreset)
-{
-	check(Collection);
-
-	if (InSubjectPreset.Role.Get() == nullptr || InSubjectPreset.Role.Get() == ULiveLinkRole::StaticClass())
-	{
-		UE_LOG(LogLiveLinkHub, Warning, TEXT("Create Subject Failure: The role is not defined."));
-		return false;
-	}
-
-	if (InSubjectPreset.Key.Source == FLiveLinkSourceCollection::DefaultVirtualSubjectGuid && InSubjectPreset.VirtualSubject == nullptr)
-	{
-		UE_LOG(LogLiveLinkHub, Warning, TEXT("Create Source Failure: Can't create an empty subject."));
-		return false;
-	}
-
-	if (InSubjectPreset.Key.SubjectName.IsNone())
-	{
-		UE_LOG(LogLiveLinkHub, Warning, TEXT("Create Subject Failure: The subject name is invalid."));
-		return false;
-	}
-
-	FLiveLinkCollectionSourceItem* SourceItem = Collection->FindSource(InSubjectPreset.Key.Source);
-	if (SourceItem == nullptr || SourceItem->bPendingKill)
-	{
-		UE_LOG(LogLiveLinkHub, Warning, TEXT("Create Subject Failure: The source doesn't exist."));
-		return false;
-	}
-
-	FLiveLinkCollectionSubjectItem* SubjectItem = Collection->FindSubject(InSubjectPreset.Key);
-	if (SubjectItem != nullptr)
-	{
-		if (SubjectItem->bPendingKill)
-		{
-			Collection->RemoveSubject(InSubjectPreset.Key);
-		}
-		else
-		{
-			UE_LOG(LogLiveLinkHub, Warning, TEXT("Create Subject Failure: The subject already exist."));
-			return false;
-		}
-	}
-
-	if (InSubjectPreset.VirtualSubject)
-	{
-		bool bEnabled = false;
-		ULiveLinkVirtualSubject* VSubject = DuplicateObject<ULiveLinkVirtualSubject>(InSubjectPreset.VirtualSubject, GetTransientPackage());
-		FLiveLinkCollectionSubjectItem VSubjectData(InSubjectPreset.Key, VSubject, bEnabled);
-		VSubject->Initialize(VSubjectData.Key, VSubject->GetRole(), this);
-
-		Collection->AddSubject(MoveTemp(VSubjectData));
-		Collection->SetSubjectEnabled(InSubjectPreset.Key, InSubjectPreset.bEnabled);
-	}
-	else
-	{
-		ULiveLinkSubjectSettings* SubjectSettings = nullptr;
-		if (InSubjectPreset.Settings)
-		{
-			SubjectSettings = DuplicateObject<ULiveLinkSubjectSettings>(InSubjectPreset.Settings, GetTransientPackage());
-		}
-		else
-		{
-			SubjectSettings = NewObject<ULiveLinkSubjectSettings>();
-		}
-
-		bool bEnabled = false;
-		FLiveLinkCollectionSubjectItem CollectionSubjectItem(InSubjectPreset.Key, MakeUnique<FLiveLinkPlaybackSubject>(SourceItem->TimedData), SubjectSettings, bEnabled);
-		CollectionSubjectItem.GetLiveSubject()->Initialize(InSubjectPreset.Key, InSubjectPreset.Role.Get(), this);
-
-		Collection->AddSubject(MoveTemp(CollectionSubjectItem));
-		Collection->SetSubjectEnabled(InSubjectPreset.Key, InSubjectPreset.bEnabled);
-	}
-	return true;
-}
-
-void FLiveLinkHubClient::PushSubjectFrameData_AnyThread(const FLiveLinkSubjectKey& SubjectKey, FLiveLinkFrameDataStruct&& FrameData)
-{
-	SCOPE_CYCLE_COUNTER(STAT_LiveLinkHub_PushFrameData);
+		TStrongObjectPtr<ULiveLinkSourceSettings> PlaybackSourceSettings = TStrongObjectPtr<ULiveLinkSourceSettings>(NewObject<ULiveLinkSourceSettings>(GetTransientPackage(), ULiveLinkHubPlaybackSourceSettings::StaticClass()));
 
-	TOptional<FLiveLinkSubjectFrameData> TranslatedFrame;
+		FLiveLinkSourcePreset ModifiedPreset = InSourcePreset;
+		ModifiedPreset.Settings = PlaybackSourceSettings.Get();
 
-	bool bFrameValid = true;
-	{
-		if (const FLiveLinkCollectionSubjectItem* SubjectItem = Collection->FindSubject(SubjectKey))
+		// Override the incoming source settings to create a playback source instead.
+		if (InSourcePreset.Settings && InSourcePreset.Settings->Factory)
 		{
-			if (FLiveLinkSubject* LiveSubject = SubjectItem->GetLiveSubject())
-			{
-				LiveSubject->SetLastPushTime(FApp::GetCurrentTime());
-				bFrameValid = LiveSubject->ValidateFrameData(FrameData);
-				if (bFrameValid)
-				{
-					LiveSubject->PreprocessFrame(FrameData);
-
-					// Translator support (Disabled for now since we don't have a way to handle the static data that's given back by the translator)
-					TArray<ULiveLinkFrameTranslator::FWorkerSharedPtr> Translators = LiveSubject->GetFrameTranslators();
-					if (Translators.Num() && Translators[0].IsValid())
-					{
-						FLiveLinkSubjectFrameData TranslatedFrameData;
-						if (Translators[0]->Translate(LiveSubject->GetStaticData(), FrameData, TranslatedFrameData))
-						{
-							TranslatedFrame = MoveTemp(TranslatedFrameData);
-						}
-					}
-				}
-			}
+			PlaybackSourceSettings->ConnectionString = InSourcePreset.Settings->Factory.GetDefaultObject()->GetSourceDisplayName().ToString();
 		}
-	}
-
-	// If the data was translated, then we broadcast this data, but we save snapshots of the original data since Evaluate will need it for translation as well.
-	FLiveLinkFrameDataStruct* FrameDataToBroadcast = TranslatedFrame.IsSet() ? &(TranslatedFrame->FrameData) : &FrameData;
-
-	if (bFrameValid)
-	{
-		OnFrameDataReceivedDelegate_AnyThread.Broadcast(SubjectKey, *FrameDataToBroadcast);
-	}
 
-	if (bVirtualSubjectsPresent)
-	{
-		// Cache frames if we have virtual subjects so they can use the data.
-		FLiveLinkClient::PushSubjectFrameData_AnyThread(SubjectKey, MoveTemp(FrameData));
+		return FLiveLinkClient::CreateSource(ModifiedPreset);
 	}
 	else
 	{
-		BroadcastFrameDataUpdate(SubjectKey, MoveTemp(*FrameDataToBroadcast));
+		return FLiveLinkClient::CreateSource(InSourcePreset);
 	}
 }
 
@@ -287,20 +80,6 @@ FText FLiveLinkHubClient::GetSourceStatus(FGuid InEntryGuid) const
 	return FLiveLinkClient::GetSourceStatus(InEntryGuid);
 }
 
-bool FLiveLinkHubClient::IsSubjectValid(const FLiveLinkSubjectKey& InSubjectKey) const
-{
-	if (const FLiveLinkCollectionSubjectItem* SubjectItem = Collection->FindSubject(InSubjectKey))
-	{
-		if (FLiveLinkSubject* LiveSubject = SubjectItem->GetLiveSubject())
-		{
-			// We don't store frame snapshots so we instead rely on the subject's last push time.
-			return FApp::GetCurrentTime() - LiveSubject->GetLastPushTime() < GetDefault<ULiveLinkSettings>()->GetTimeWithoutFrameToBeConsiderAsInvalid();
-		}
-		return true;
-	}
-	return false;
-}
-
 void FLiveLinkHubClient::RemoveSubject_AnyThread(const FLiveLinkSubjectKey& InSubjectKey)
 {
 	OnSubjectMarkedPendingKill_AnyThread().Broadcast(InSubjectKey);
@@ -339,44 +118,18 @@ TSharedPtr<ILiveLinkProvider> FLiveLinkHubClient::GetRebroadcastLiveLinkProvider
 
 void FLiveLinkHubClient::BroadcastStaticDataUpdate(FLiveLinkSubject* InLiveSubject, TSubclassOf<ULiveLinkRole> InRole, const FLiveLinkStaticDataStruct& InStaticData) const
 {
-	// If we have a translator, we need to broadcast the translated static data to the clients.
-	FLiveLinkStaticDataStruct StaticDataStruct;
-	TArray<ULiveLinkFrameTranslator::FWorkerSharedPtr> Translators = InLiveSubject->GetFrameTranslators();
-	if (Translators.Num() && Translators[0].IsValid())
-	{
-
-		ULiveLinkFrameTranslator::FWorkerSharedPtr Translator = Translators[0];
-		TSubclassOf<ULiveLinkRole> FromRole = Translator->GetFromRole();
-		UScriptStruct* FrameDataStruct = FromRole->GetDefaultObject<ULiveLinkRole>()->GetFrameDataStruct();
-		FStructOnScope FrameDataOnScope(FrameDataStruct);
-		FrameDataStruct->InitializeDefaultValue(FrameDataOnScope.GetStructMemory());
-
-		FLiveLinkFrameDataStruct BaseFrameStruct(FrameDataStruct);
-		FrameDataStruct->InitializeDefaultValue(FrameDataOnScope.GetStructMemory());
-
-		FLiveLinkSubjectFrameData OutTranslatedFrame;
-		Translator->Translate(InLiveSubject->GetStaticData(), BaseFrameStruct, OutTranslatedFrame);
-
-		OnStaticDataReceivedDelegate_AnyThread.Broadcast(InLiveSubject->GetSubjectKey(), Translator->GetToRole(), OutTranslatedFrame.StaticData);
-	}
-	else
-	{
-		OnStaticDataReceivedDelegate_AnyThread.Broadcast(InLiveSubject->GetSubjectKey(), InRole, InStaticData);
-	}
+	OnStaticDataReceivedDelegate_AnyThread.Broadcast(InLiveSubject->GetSubjectKey(), InRole, InStaticData);
 }
 
 void FLiveLinkHubClient::OnStaticDataAdded(FLiveLinkSubjectKey SubjectKey,  TSubclassOf<ULiveLinkRole> SubjectRole, const FLiveLinkStaticDataStruct& InStaticData)
 {
-	// Broadcast the subjects to connected clients.
-	if (FLiveLinkCollectionSubjectItem* SubjectItem = Collection->FindSubject(SubjectKey))
-	{
-		BroadcastStaticDataUpdate(SubjectItem->GetLiveSubject(), SubjectRole, InStaticData);
-	}
+	OnStaticDataReceivedDelegate_AnyThread.Broadcast(SubjectKey, SubjectRole, InStaticData);
 }
 
 void FLiveLinkHubClient::OnFrameDataAdded(FLiveLinkSubjectKey InSubjectKey, TSubclassOf<ULiveLinkRole> SubjectRole, const FLiveLinkFrameDataStruct& InFrameData)
 {
-	// No-op, to be used in subsequent CL.
+	OnFrameDataReceivedDelegate_AnyThread.Broadcast(InSubjectKey, InFrameData);
 }
 
+
 #undef LOCTEXT_NAMESPACE
diff --git a/Engine/Plugins/Animation/LiveLinkHub/Source/LiveLinkHub/Private/LiveLinkHubClient.h b/Engine/Plugins/Animation/LiveLinkHub/Source/LiveLinkHub/Private/LiveLinkHubClient.h
--- a/Engine/Plugins/Animation/LiveLinkHub/Source/LiveLinkHub/Private/LiveLinkHubClient.h
+++ b/Engine/Plugins/Animation/LiveLinkHub/Source/LiveLinkHub/Private/LiveLinkHubClient.h
@@ -14,7 +14,7 @@ class FLiveLinkSubject;
 class ILiveLinkHub;
 struct ILiveLinkProvider;
 
-DECLARE_TS_MULTICAST_DELEGATE_TwoParams(FOnFrameDataReceived_AnyThread, const FLiveLinkSubjectKey&, FLiveLinkFrameDataStruct&);
+DECLARE_TS_MULTICAST_DELEGATE_TwoParams(FOnFrameDataReceived_AnyThread, const FLiveLinkSubjectKey&, const FLiveLinkFrameDataStruct&);
 DECLARE_TS_MULTICAST_DELEGATE_ThreeParams(FOnStaticDataReceived_AnyThread, const FLiveLinkSubjectKey&, TSubclassOf<ULiveLinkRole>, const FLiveLinkStaticDataStruct&);
 DECLARE_TS_MULTICAST_DELEGATE_OneParam(FOnSubjectMarkedPendingKill_AnyThread, const FLiveLinkSubjectKey&);
 
@@ -52,10 +52,7 @@ public:
 public:
 	//~ Begin ILiveLinkClient interface
 	virtual bool CreateSource(const FLiveLinkSourcePreset& InSourcePreset) override;
-	virtual bool CreateSubject(const FLiveLinkSubjectPreset& InSubjectPreset) override;
-	virtual void PushSubjectFrameData_AnyThread(const FLiveLinkSubjectKey& SubjectKey, FLiveLinkFrameDataStruct&& FrameData) override;
 	virtual FText GetSourceStatus(FGuid InEntryGuid) const override;
-	virtual bool IsSubjectValid(const FLiveLinkSubjectKey& InSubjectKey) const override;
 	virtual void RemoveSubject_AnyThread(const FLiveLinkSubjectKey& InSubjectKey) override;
 	virtual bool AddVirtualSubject(const FLiveLinkSubjectKey& VirtualSubjectKey, TSubclassOf<ULiveLinkVirtualSubject> VirtualSubjectClass) override;
     virtual void RemoveVirtualSubject(const FLiveLinkSubjectKey& VirtualSubjectKey) override;
@@ -66,10 +63,6 @@ public:
 	//~ End FLiveLinkClient interface
 
 private:
-	/** Create a LiveLinkPlaybackSource which acts as a dummy source when doing playback. */
-	bool CreatePlaybackSource(const FLiveLinkSourcePreset& InSourcePreset);
-	/** Create a LiveLinkPlaybackSubject which acts as a dummy subject when doing playback. */
-	bool CreatePlaybackSubject(const FLiveLinkSubjectPreset& InSubjectPreset);
 	/** Broadcast a static data update to this client's listeners. */
 	void BroadcastStaticDataUpdate(FLiveLinkSubject* InLiveSubject, TSubclassOf<ULiveLinkRole> InRole, const FLiveLinkStaticDataStruct& InStaticData) const;
 
diff --git a/Engine/Plugins/Animation/LiveLinkHub/Source/LiveLinkHub/Private/LiveLinkHubSettings.h b/Engine/Plugins/Animation/LiveLinkHub/Source/LiveLinkHub/Private/LiveLinkHubSettings.h
--- a/Engine/Plugins/Animation/LiveLinkHub/Source/LiveLinkHub/Private/LiveLinkHubSettings.h
+++ b/Engine/Plugins/Animation/LiveLinkHub/Source/LiveLinkHub/Private/LiveLinkHubSettings.h
@@ -39,5 +39,5 @@ public:
 
 	/** Target framerate for ticking LiveLinkHub. */
 	UPROPERTY(config, EditAnywhere, Category="LiveLinkHub", meta = (ConfigRestartRequired = true, ClampMin="15.0"))
-	float TargetFrameRate = 120.0f;
+	float TargetFrameRate = 60.0f;
 };
diff --git a/Engine/Plugins/Animation/LiveLinkHub/Source/LiveLinkHub/Private/Recording/LiveLinkHubPlaybackSourceFactory.h b/Engine/Plugins/Animation/LiveLinkHub/Source/LiveLinkHub/Private/Recording/LiveLinkHubPlaybackSourceFactory.h
new file mode 100644
--- /dev/null
+++ b/Engine/Plugins/Animation/LiveLinkHub/Source/LiveLinkHub/Private/Recording/LiveLinkHubPlaybackSourceFactory.h
@@ -0,0 +1,38 @@
+// Copyright Epic Games, Inc. All Rights Reserved.
+
+#pragma once
+
+#include "LiveLinkSourceFactory.h"
+
+#include "LiveLinkPlaybackSource.h"
+
+#include "LiveLinkHubPlaybackSourceFactory.generated.h"
+
+#define LOCTEXT_NAMESPACE "LiveLinkHubPlaybackSourceFactory"
+
+/** Factory used to create playback sources when we're playing back a livelink recording. */
+UCLASS()
+class ULiveLinkHubPlaybackSourceFactory : public ULiveLinkSourceFactory
+{
+public:
+	GENERATED_BODY()
+
+	//~ Begin ULiveLinkSourceFactory interface
+	virtual FText GetSourceDisplayName() const override
+	{
+		return LOCTEXT("PlaybackSourceDisplayName", "Playback Source");
+	}
+
+	virtual FText GetSourceTooltip() const
+	{
+		return LOCTEXT("PlaybackSourceTooltip", "Mock livelink source used to playback livelink data that was recorded.");
+	}
+
+	virtual TSharedPtr<ILiveLinkSource> CreateSource(const FString& ConnectionString) const
+	{
+		return MakeShared<FLiveLinkPlaybackSource>(ConnectionString);
+	}
+	//~ End ULiveLinkSourceFactory interface
+};
+
+#undef LOCTEXT_NAMESPACE
diff --git a/Engine/Plugins/Animation/LiveLinkHub/Source/LiveLinkHub/Private/Recording/LiveLinkHubPlaybackSourceSettings.h b/Engine/Plugins/Animation/LiveLinkHub/Source/LiveLinkHub/Private/Recording/LiveLinkHubPlaybackSourceSettings.h
new file mode 100644
--- /dev/null
+++ b/Engine/Plugins/Animation/LiveLinkHub/Source/LiveLinkHub/Private/Recording/LiveLinkHubPlaybackSourceSettings.h
@@ -0,0 +1,23 @@
+// Copyright Epic Games, Inc. All Rights Reserved.
+
+#pragma once
+
+#include "LiveLinkSourceSettings.h"
+
+#include "LiveLinkHubPlaybackSourceFactory.h"
+#include "LiveLinkSourceFactory.h"
+
+#include "LiveLinkHubPlaybackSourceSettings.generated.h"
+
+/** Settings class for playback sources, which provides the factory class for them. */
+UCLASS()
+class ULiveLinkHubPlaybackSourceSettings : public ULiveLinkSourceSettings
+{
+public:
+	GENERATED_BODY()
+
+	ULiveLinkHubPlaybackSourceSettings()
+	{
+		Factory = ULiveLinkHubPlaybackSourceFactory::StaticClass();
+	}
+};
diff --git a/Engine/Plugins/Animation/LiveLinkHub/Source/LiveLinkHub/Private/Recording/LiveLinkPlaybackSource.h b/Engine/Plugins/Animation/LiveLinkHub/Source/LiveLinkHub/Private/Recording/LiveLinkPlaybackSource.h
--- a/Engine/Plugins/Animation/LiveLinkHub/Source/LiveLinkHub/Private/Recording/LiveLinkPlaybackSource.h
+++ b/Engine/Plugins/Animation/LiveLinkHub/Source/LiveLinkHub/Private/Recording/LiveLinkPlaybackSource.h
@@ -5,8 +5,6 @@
 #include "ILiveLinkSource.h"
 #include "LiveLinkSourceSettings.h"
 
-#include "LiveLinkPlaybackSource.generated.h"
-
 class ILiveLinkClient;
 
 /**
@@ -16,6 +14,12 @@ class FLiveLinkPlaybackSource : public ILiveLinkSource
 {
 public:
 	FLiveLinkPlaybackSource() = default;
+
+	FLiveLinkPlaybackSource(const FString& OriginalSourceName)
+	{
+		SourceName = *FString::Printf(TEXT("Playback (%s)"), *OriginalSourceName);
+	}
+
 	virtual ~FLiveLinkPlaybackSource() = default;
 
 	//~ Begin ILiveLinkSource interface
@@ -30,7 +34,6 @@ public:
 	
 	virtual void InitializeSettings(ULiveLinkSourceSettings* Settings) override
 	{
-		SourceName = TEXT("Playback source.");
 	}
 	
 	virtual bool IsSourceStillValid() const override
@@ -63,16 +66,3 @@ protected:
 	/** Source name. */
 	FName SourceName;
 };
-
-/** PlaybackSourceSettings to be able to differentiate from live sources and keep a name associated to the source */
-UCLASS()
-class ULiveLinkPlaybackSourceSettings : public ULiveLinkSourceSettings
-{
-	GENERATED_BODY()
-
-public:
-
-	/** Source name. */
-	UPROPERTY()
-	FName SourceName;
-};
diff --git a/Engine/Plugins/Animation/LiveLinkHub/Source/LiveLinkHub/Private/Recording/LiveLinkPlaybackSubject.h b/Engine/Plugins/Animation/LiveLinkHub/Source/LiveLinkHub/Private/Recording/LiveLinkPlaybackSubject.h
deleted file mode 100644
diff --git a/Engine/Programs/LiveLinkHub/Config/DefaultEngine.ini b/Engine/Programs/LiveLinkHub/Config/DefaultEngine.ini
--- a/Engine/Programs/LiveLinkHub/Config/DefaultEngine.ini
+++ b/Engine/Programs/LiveLinkHub/Config/DefaultEngine.ini
@@ -28,6 +28,12 @@ GameNameShort=LLH
 [/Script/MacTargetPlatform.XcodeProjectSettings]
 TemplateMacPlist=(FilePath="/Engine/Source/Programs/LiveLinkHub/Resources/Mac/Info.plist")
 
+[/Script/LiveLinkInterface.LiveLinkDefaultSourceSettings]
+DefaultSourceFrameBufferSize=1
+
+[/Script/LiveLinkInterface.LiveLinkDefaultSubjectSettings]
+bRebroadcastSubjectsByDefault=true
+bAllowModifyingRebroadcast=false
 
 [/Script/LiveLinkHub.LiveLinkHubSettings]
 bAutoAddDiscoveredClients=true
diff --git a/Engine/Programs/LiveLinkHub/Config/DefaultGame.ini b/Engine/Programs/LiveLinkHub/Config/DefaultGame.ini
--- a/Engine/Programs/LiveLinkHub/Config/DefaultGame.ini
+++ b/Engine/Programs/LiveLinkHub/Config/DefaultGame.ini
@@ -8,14 +8,10 @@ MessageBusPingRequestFrequency=1.000000
 MessageBusHeartbeatFrequency=1.000000
 MessageBusHeartbeatTimeout=2.000000
 MessageBusTimeBeforeRemovingInactiveSource=30.000000
+bPreProcessRebroadcastedFrames=true
+bTranslateRebroadcastedFrames=true
 TimeWithoutFrameToBeConsiderAsInvalid=0.500000
 ValidColor=(R=0.000000,G=1.000000,B=0.000000,A=1.000000)
 InvalidColor=(R=1.000000,G=1.000000,B=0.000000,A=1.000000)
 TextSizeSource=16
 TextSizeSubject=12
-
-[/Script/LiveLinkHub.LiveLinkHubSettings]
-bAutoAddDiscoveredClients=True
-+ProjectSettingsToShow=LiveLink
-+ProjectSettingsToShow=LiveLinkHub
-
