LiveLinkHub - Fix case where a freeze could occur when streaming the recording file in reverse, and the file was at the end.

diff --git a/Engine/Plugins/Animation/LiveLinkHub/Source/LiveLinkHub/Private/Recording/Implementations/LiveLinkUAssetRecording.cpp b/Engine/Plugins/Animation/LiveLinkHub/Source/LiveLinkHub/Private/Recording/Implementations/LiveLinkUAssetRecording.cpp
--- a/Engine/Plugins/Animation/LiveLinkHub/Source/LiveLinkHub/Private/Recording/Implementations/LiveLinkUAssetRecording.cpp
+++ b/Engine/Plugins/Animation/LiveLinkHub/Source/LiveLinkHub/Private/Recording/Implementations/LiveLinkUAssetRecording.cpp
@@ -343,85 +343,81 @@ void ULiveLinkUAssetRecording::LoadRecordingAsync(int32 InStartFrame, int32 InCu
 
 	DebugSleepTime = CVarLiveLinkHubDebugFrameBufferDelay.GetValueOnAnyThread();
 	
-	// Read file for streaming into memory.
-	if (!RecordingFileReader->AtEnd())
+	// Perform initial load and record entry frame file offsets.
+	const bool bInitialLoad = FrameFileData.Num() == 0;
+	if (bInitialLoad)
 	{
-		// Perform initial load and record entry frame file offsets.
-		const bool bInitialLoad = FrameFileData.Num() == 0;
-		if (bInitialLoad)
-		{
-			RecordingFileReader->Seek(0);
+		RecordingFileReader->Seek(0);
 		
-			int32 LoadedRecordingVersion;
-			*RecordingFileReader << LoadedRecordingVersion;
+		int32 LoadedRecordingVersion;
+		*RecordingFileReader << LoadedRecordingVersion;
 
-			// If we modify the RecordingVersion we can perform import logic here.
-			ensure(LoadedRecordingVersion == RecordingVersion);
+		// If we modify the RecordingVersion we can perform import logic here.
+		ensure(LoadedRecordingVersion == RecordingVersion);
 			
-			// Process static data.
+		// Process static data.
 		
-			int32 NumStaticData = 0;
-			*RecordingFileReader << NumStaticData;
+		int32 NumStaticData = 0;
+		*RecordingFileReader << NumStaticData;
 
-			for (int32 StaticIdx = 0; StaticIdx < NumStaticData; ++StaticIdx)
-			{
-				FGuid KeySource = FGuid();
-				FString KeyName;
+		for (int32 StaticIdx = 0; StaticIdx < NumStaticData; ++StaticIdx)
+		{
+			FGuid KeySource = FGuid();
+			FString KeyName;
 			
-				*RecordingFileReader << KeySource;
-				*RecordingFileReader << KeyName;
+			*RecordingFileReader << KeySource;
+			*RecordingFileReader << KeyName;
 
-				// Create framedata just to load initial static frame data. Static data doesn't require this afterward.
-				FFrameFileData TemporaryFrameData;
-				TemporaryFrameData.FrameDataSubjectKey = MakeShared<FLiveLinkSubjectKey>(KeySource, *KeyName);
-				if (!LoadInitialFrameData(TemporaryFrameData))
-				{
-					return;
-				}
-				
-				FLiveLinkRecordingStaticDataContainer& DataContainer = RecordingData.StaticData.FindChecked(*TemporaryFrameData.FrameDataSubjectKey.Get());
-				LoadFrameData(TemporaryFrameData, DataContainer, 0, 0, 1);
+			// Create framedata just to load initial static frame data. Static data doesn't require this afterward.
+			FFrameFileData TemporaryFrameData;
+			TemporaryFrameData.FrameDataSubjectKey = MakeShared<FLiveLinkSubjectKey>(KeySource, *KeyName);
+			if (!LoadInitialFrameData(TemporaryFrameData))
+			{
+				return;
 			}
+				
+			FLiveLinkRecordingStaticDataContainer& DataContainer = RecordingData.StaticData.FindChecked(*TemporaryFrameData.FrameDataSubjectKey.Get());
+			LoadFrameData(TemporaryFrameData, DataContainer, 0, 0, 1);
+		}
 
-			// Process frame data.
+		// Process frame data.
 		
-			int32 NumFrameData = 0;
-			*RecordingFileReader << NumFrameData;
+		int32 NumFrameData = 0;
+		*RecordingFileReader << NumFrameData;
 
-			for (int32 FrameIdx = 0; FrameIdx < NumFrameData; ++FrameIdx)
-			{
-				FGuid KeySource = FGuid();
-				FString KeyName;
+		for (int32 FrameIdx = 0; FrameIdx < NumFrameData; ++FrameIdx)
+		{
+			FGuid KeySource = FGuid();
+			FString KeyName;
 			
-				*RecordingFileReader << KeySource;
-				*RecordingFileReader << KeyName;
+			*RecordingFileReader << KeySource;
+			*RecordingFileReader << KeyName;
 
-				FFrameFileData KeyPosition;
-				KeyPosition.FrameDataSubjectKey = MakeShared<FLiveLinkSubjectKey>(KeySource, *KeyName);
-				if (!LoadInitialFrameData(KeyPosition))
-				{
-					return;
-				}
-
-				// Offset to the end of this block if there is multiple NumFrameData.
-				const int64 EndBlockPosition = KeyPosition.GetFrameFilePosition(KeyPosition.MaxFrames);
-				RecordingFileReader->Seek(EndBlockPosition);
-				FrameFileData.Add(MoveTemp(KeyPosition));
+			FFrameFileData KeyPosition;
+			KeyPosition.FrameDataSubjectKey = MakeShared<FLiveLinkSubjectKey>(KeySource, *KeyName);
+			if (!LoadInitialFrameData(KeyPosition))
+			{
+				return;
 			}
+
+			// Offset to the end of this block if there is multiple NumFrameData.
+			const int64 EndBlockPosition = KeyPosition.GetFrameFilePosition(KeyPosition.MaxFrames);
+			RecordingFileReader->Seek(EndBlockPosition);
+			FrameFileData.Add(MoveTemp(KeyPosition));
 		}
+	}
 
-		// Load the required frames, either on initial load or subsequent loads.
-		for (FFrameFileData& FrameData : FrameFileData)
+	// Load the required frames, either on initial load or subsequent loads.
+	for (FFrameFileData& FrameData : FrameFileData)
+	{
+		if (ensure(FrameData.FrameDataSubjectKey.IsValid()))
 		{
-			if (ensure(FrameData.FrameDataSubjectKey.IsValid()))
-			{
-				FLiveLinkRecordingBaseDataContainer& DataContainer = RecordingData.FrameData.FindChecked(*FrameData.FrameDataSubjectKey.Get());
-				LoadFrameData(FrameData, DataContainer, InStartFrame, InCurrentFrame, InNumFramesToLoad);
-			}
-			else
-			{
-				UE_LOG(LogLiveLinkHub, Error, TEXT("FrameDataSubjectKey is missing for file %s."), *GetRecordingDataFilePath());
-			}
+			FLiveLinkRecordingBaseDataContainer& DataContainer = RecordingData.FrameData.FindChecked(*FrameData.FrameDataSubjectKey.Get());
+			LoadFrameData(FrameData, DataContainer, InStartFrame, InCurrentFrame, InNumFramesToLoad);
+		}
+		else
+		{
+			UE_LOG(LogLiveLinkHub, Error, TEXT("FrameDataSubjectKey is missing for file %s."), *GetRecordingDataFilePath());
 		}
 	}
 }
