LiveLinkHub - Support filename templates and automatic tokens. Several date and time tokens can now be used in LLH settings, and sessionName is now recognized.

The .rec file will now be stored in the same location as the uasset.

diff --git a/Engine/Plugins/Animation/LiveLinkHub/Source/LiveLinkHub/Private/Clients/LiveLinkHubProvider.cpp b/Engine/Plugins/Animation/LiveLinkHub/Source/LiveLinkHub/Private/Clients/LiveLinkHubProvider.cpp
--- a/Engine/Plugins/Animation/LiveLinkHub/Source/LiveLinkHub/Private/Clients/LiveLinkHubProvider.cpp
+++ b/Engine/Plugins/Animation/LiveLinkHub/Source/LiveLinkHub/Private/Clients/LiveLinkHubProvider.cpp
@@ -14,12 +14,12 @@
 #include "HAL/CriticalSection.h"
 #include "LiveLinkHubLog.h"
 #include "LiveLinkHubMessages.h"
-#include "LiveLinkHubSettings.h"
 #include "LiveLinkSettings.h"
 #include "MessageEndpointBuilder.h"
 #include "Misc/ScopeLock.h"
 #include "Session/LiveLinkHubSession.h"
 #include "Session/LiveLinkHubSessionManager.h"
+#include "Settings/LiveLinkHubSettings.h"
 #include "TimerManager.h"
 
 
diff --git a/Engine/Plugins/Animation/LiveLinkHub/Source/LiveLinkHub/Private/Config/LiveLinkHubFileUtilities.cpp b/Engine/Plugins/Animation/LiveLinkHub/Source/LiveLinkHub/Private/Config/LiveLinkHubFileUtilities.cpp
--- a/Engine/Plugins/Animation/LiveLinkHub/Source/LiveLinkHub/Private/Config/LiveLinkHubFileUtilities.cpp
+++ b/Engine/Plugins/Animation/LiveLinkHub/Source/LiveLinkHub/Private/Config/LiveLinkHubFileUtilities.cpp
@@ -2,10 +2,13 @@
 
 #include "LiveLinkHubFileUtilities.h"
 
+#include "Config/LiveLinkHubTemplateTokens.h"
 #include "HAL/FileManager.h"
 #include "JsonObjectConverter.h"
+#include "LiveLinkHub.h"
 #include "LiveLinkHubLog.h"
 #include "Session/LiveLinkHubSessionData.h"
+#include "Session/LiveLinkHubSessionManager.h"
 #include "UObject/Package.h"
 
 void UE::LiveLinkHub::FileUtilities::Private::SaveConfig(const ULiveLinkHubSessionData* InConfigData, const FString& InFilePath)
@@ -83,3 +86,59 @@ ULiveLinkHubSessionData* UE::LiveLinkHub::FileUtilities::Private::FromJson(const
 
 	return nullptr;
 }
+
+void UE::LiveLinkHub::FileUtilities::Private::ParseFilenameTemplate(const FString& InFilenameTemplate,
+	FFilenameTemplateData& OutTemplateData)
+{
+	FString FormattedString = InFilenameTemplate;
+	
+	// Replace tokens.
+	{
+		// Get current datetime.
+		const FDateTime CurrentDate = FDateTime::Now();
+		const FString Year2DigitValue = FString::Printf(TEXT("%02d"), CurrentDate.GetYear() % 100);
+		const FString Year4DigitValue = FString::Printf(TEXT("%02d"), CurrentDate.GetYear());
+		const FString MonthValue = FString::Printf(TEXT("%02d"), CurrentDate.GetMonth());
+		const FString DayValue = FString::Printf(TEXT("%02d"), CurrentDate.GetDay());
+		const FString HourValue = FString::Printf(TEXT("%02d"), CurrentDate.GetHour());
+		const FString MinuteValue = FString::Printf(TEXT("%02d"), CurrentDate.GetMinute());
+
+		const FLiveLinkHubAutomaticTokens& AutomaticTokens = FLiveLinkHubAutomaticTokens::GetStaticTokens();
+
+		using namespace UE::LiveLinkHub::Tokens::Private;
+		
+		FormattedString.ReplaceInline(*CreateToken(AutomaticTokens.Year4Digit), *Year4DigitValue, ESearchCase::IgnoreCase);
+		FormattedString.ReplaceInline(*CreateToken(AutomaticTokens.Year2Digit), *Year2DigitValue, ESearchCase::IgnoreCase);
+		FormattedString.ReplaceInline(*CreateToken(AutomaticTokens.Month), *MonthValue, ESearchCase::CaseSensitive);
+		FormattedString.ReplaceInline(*CreateToken(AutomaticTokens.Day), *DayValue, ESearchCase::IgnoreCase);
+		FormattedString.ReplaceInline(*CreateToken(AutomaticTokens.Hour), *HourValue, ESearchCase::IgnoreCase);
+		FormattedString.ReplaceInline(*CreateToken(AutomaticTokens.Minute), *MinuteValue, ESearchCase::CaseSensitive);
+
+		// Get session information.
+		FString SessionNameValue;
+		if (TSharedPtr<ILiveLinkHubSessionManager> SessionManager = FModuleManager::Get().GetModuleChecked<FLiveLinkHubModule>("LiveLinkHub").GetLiveLinkHub()->GetSessionManager())
+		{
+			SessionNameValue = FPaths::GetBaseFilename(SessionManager->GetLastConfigPath());
+		}
+		FormattedString.ReplaceInline(*CreateToken(AutomaticTokens.SessionName), *SessionNameValue);
+	}
+	
+	OutTemplateData.FullPath = FormattedString;
+
+	// Split folder path and file name.
+	int32 LastSlashIndex;
+	if (FormattedString.FindLastChar('/', LastSlashIndex))
+	{
+		OutTemplateData.FolderPath = FormattedString.Left(LastSlashIndex);
+		if (OutTemplateData.FolderPath.StartsWith("/"))
+		{
+			OutTemplateData.FolderPath.RemoveFromStart("/");
+		}
+		OutTemplateData.FileName = FormattedString.Mid(LastSlashIndex + 1);
+	}
+	else
+	{
+		OutTemplateData.FolderPath = TEXT("");
+		OutTemplateData.FileName = FormattedString;
+	}
+}
diff --git a/Engine/Plugins/Animation/LiveLinkHub/Source/LiveLinkHub/Private/Config/LiveLinkHubFileUtilities.h b/Engine/Plugins/Animation/LiveLinkHub/Source/LiveLinkHub/Private/Config/LiveLinkHubFileUtilities.h
--- a/Engine/Plugins/Animation/LiveLinkHub/Source/LiveLinkHub/Private/Config/LiveLinkHubFileUtilities.h
+++ b/Engine/Plugins/Animation/LiveLinkHub/Source/LiveLinkHub/Private/Config/LiveLinkHubFileUtilities.h
@@ -33,4 +33,18 @@ namespace UE::LiveLinkHub::FileUtilities::Private
 
 	/** Convert config data from json. */
 	ULiveLinkHubSessionData* FromJson(const TSharedPtr<FJsonObject>& InJsonObject);
+
+	/** Evaluated results from a template file string. */
+	struct FFilenameTemplateData
+	{
+		/** Full original path with evaluated tokens. */
+		FString FullPath;
+		/** The evaluated folder(s) without the first forward slash or filename. */
+		FString FolderPath;
+		/** The evaluated filename. */
+		FString FileName;
+	};
+
+	/** Parse a filename template for included template tokens and output the computed result. */
+	void ParseFilenameTemplate(const FString& InFilenameTemplate, FFilenameTemplateData& OutTemplateData);
 }
diff --git a/Engine/Plugins/Animation/LiveLinkHub/Source/LiveLinkHub/Private/Config/LiveLinkHubTemplateTokens.h b/Engine/Plugins/Animation/LiveLinkHub/Source/LiveLinkHub/Private/Config/LiveLinkHubTemplateTokens.h
new file mode 100644
--- /dev/null
+++ b/Engine/Plugins/Animation/LiveLinkHub/Source/LiveLinkHub/Private/Config/LiveLinkHubTemplateTokens.h
@@ -0,0 +1,62 @@
+// Copyright Epic Games, Inc. All Rights Reserved.
+
+#pragma once
+
+#include "CoreMinimal.h"
+#include "UObject/ObjectMacros.h"
+
+#include "LiveLinkHubTemplateTokens.generated.h"
+
+/**
+ * Predefined tokens used in template parsing. Set up as a USTRUCT so the properties and tooltips
+ * can be easily iterated and displayed to the user.
+ */
+USTRUCT()
+struct FLiveLinkHubAutomaticTokens
+{
+	GENERATED_BODY()
+
+	/** Two-digit year. */
+	UPROPERTY()
+	FString Year2Digit = TEXT("yy");
+	
+	/** Four-digit year. */
+	UPROPERTY()
+	FString Year4Digit = TEXT("yyyy");
+	
+	/** Two-digit month. */
+	UPROPERTY()
+	FString Month = TEXT("MM");
+	
+	/** Two-digit day. */
+	UPROPERTY()
+	FString Day = TEXT("dd");
+	
+	/** Two-digit hour. */
+	UPROPERTY()
+	FString Hour = TEXT("hh");
+	
+	/** Two-digit minute. */
+	UPROPERTY()
+	FString Minute = TEXT("mm");
+
+	/** The current session name. */
+	UPROPERTY()
+	FString SessionName = TEXT("sessionName");
+
+	/** Singleton like access to this ustruct. */
+	static FLiveLinkHubAutomaticTokens& GetStaticTokens()
+	{
+		static FLiveLinkHubAutomaticTokens AutomaticTokens;
+		return AutomaticTokens;
+	}
+};
+
+namespace UE::LiveLinkHub::Tokens::Private
+{
+	/** Return the token formatted with curly brackets. */
+	inline FString CreateToken(const FString& InString)
+	{
+		return FString::Printf(TEXT("{%s}"), *InString);
+	}
+}
\ No newline at end of file
diff --git a/Engine/Plugins/Animation/LiveLinkHub/Source/LiveLinkHub/Private/LiveLinkHub.cpp b/Engine/Plugins/Animation/LiveLinkHub/Source/LiveLinkHub/Private/LiveLinkHub.cpp
--- a/Engine/Plugins/Animation/LiveLinkHub/Source/LiveLinkHub/Private/LiveLinkHub.cpp
+++ b/Engine/Plugins/Animation/LiveLinkHub/Source/LiveLinkHub/Private/LiveLinkHub.cpp
@@ -12,9 +12,8 @@
 #include "ISettingsModule.h"
 #include "LiveLinkEditorSettings.h"
 #include "LiveLinkHubClient.h"
-#include "LiveLinkHubSettings.h"
-#include "LiveLinkHubSubjectSettings.h"
 #include "LiveLinkHubCommands.h"
+#include "LiveLinkHubSubjectSettings.h"
 #include "LiveLinkProviderImpl.h"
 #include "LiveLinkSettings.h"
 #include "Misc/App.h"
@@ -23,6 +22,7 @@
 #include "Recording/LiveLinkHubRecordingController.h"
 #include "Recording/LiveLinkHubRecordingListController.h"
 #include "Session/LiveLinkHubSessionManager.h"
+#include "Settings/LiveLinkHubSettings.h"
 #include "Subjects/LiveLinkHubSubjectController.h"
 #include "UI/Window/LiveLinkHubWindowController.h"
 
diff --git a/Engine/Plugins/Animation/LiveLinkHub/Source/LiveLinkHub/Private/LiveLinkHubApplication.h b/Engine/Plugins/Animation/LiveLinkHub/Source/LiveLinkHub/Private/LiveLinkHubApplication.h
--- a/Engine/Plugins/Animation/LiveLinkHub/Source/LiveLinkHub/Private/LiveLinkHubApplication.h
+++ b/Engine/Plugins/Animation/LiveLinkHub/Source/LiveLinkHub/Private/LiveLinkHubApplication.h
@@ -2,18 +2,18 @@
 
 #pragma once
 
-#include "Commandlets/Commandlet.h"
 #include "CoreMinimal.h"
+#include "Commandlets/Commandlet.h"
 #include "Editor/EditorPerformanceSettings.h"
 #include "Engine/Engine.h"
 #include "Framework/Application/SlateApplication.h"
 #include "LiveLinkHub.h"
 #include "LiveLinkHubInputProcessor.h"
-#include "LiveLinkHubSettings.h"
 #include "LiveLinkHubTicker.h"
 #include "Misc/App.h"
 #include "Misc/CoreDelegates.h"
 #include "Runtime/Launch/Resources/Version.h"
+#include "Settings/LiveLinkHubSettings.h"
 #include "UObject/UObjectGlobals.h"
 
 #if PLATFORM_MAC
diff --git a/Engine/Plugins/Animation/LiveLinkHub/Source/LiveLinkHub/Private/LiveLinkHubModule.cpp b/Engine/Plugins/Animation/LiveLinkHub/Source/LiveLinkHub/Private/LiveLinkHubModule.cpp
--- a/Engine/Plugins/Animation/LiveLinkHub/Source/LiveLinkHub/Private/LiveLinkHubModule.cpp
+++ b/Engine/Plugins/Animation/LiveLinkHub/Source/LiveLinkHub/Private/LiveLinkHubModule.cpp
@@ -10,6 +10,8 @@
 #include "PropertyEditorModule.h"
 #include "Recording/LiveLinkHubPlaybackController.h"
 #include "Recording/LiveLinkHubRecordingController.h"
+#include "Settings/LiveLinkHubSettings.h"
+#include "Settings/LiveLinkHubSettingsCustomization.h"
 #include "Subjects/LiveLinkHubSubjectSettingsDetailsCustomization.h"
 
 #if !WITH_LIVELINK_HUB
@@ -49,7 +51,10 @@ void FLiveLinkHubModule::StartupModule()
 	FPropertyEditorModule& PropertyModule = FModuleManager::LoadModuleChecked<FPropertyEditorModule>("PropertyEditor");
 	PropertyModule.RegisterCustomClassLayout(ULiveLinkHubSubjectSettings::StaticClass()->GetFName(), FOnGetDetailCustomizationInstance::CreateStatic(&FLiveLinkHubSubjectSettingsDetailsCustomization::MakeInstance));
 #endif
-
+	if (FPropertyEditorModule* PropertyEditorModule = FModuleManager::GetModulePtr<FPropertyEditorModule>("PropertyEditor"))
+	{
+		PropertyEditorModule->RegisterCustomClassLayout(ULiveLinkHubSettings::StaticClass()->GetFName(), FOnGetDetailCustomizationInstance::CreateStatic(&FLiveLinkHubSettingsCustomization::MakeInstance));
+	}
 }
 
 void FLiveLinkHubModule::ShutdownModule()
@@ -62,6 +67,10 @@ void FLiveLinkHubModule::ShutdownModule()
 		PropertyModule->UnregisterCustomClassLayout(ULiveLinkHubSubjectSettings::StaticClass()->GetFName());
 	}
 #endif
+	if (FPropertyEditorModule* PropertyEditorModule = FModuleManager::GetModulePtr<FPropertyEditorModule>("PropertyEditor"))
+	{
+		PropertyEditorModule->UnregisterCustomClassLayout(ULiveLinkHubSettings::StaticClass()->GetFName());
+	}
 }
 
 TSharedPtr<FLiveLinkHub> FLiveLinkHubModule::GetLiveLinkHub() const
diff --git a/Engine/Plugins/Animation/LiveLinkHub/Source/LiveLinkHub/Private/LiveLinkHubSettings.h b/Engine/Plugins/Animation/LiveLinkHub/Source/LiveLinkHub/Private/LiveLinkHubSettings.h
deleted file mode 100644
diff --git a/Engine/Plugins/Animation/LiveLinkHub/Source/LiveLinkHub/Private/LiveLinkHubTicker.h b/Engine/Plugins/Animation/LiveLinkHub/Source/LiveLinkHub/Private/LiveLinkHubTicker.h
--- a/Engine/Plugins/Animation/LiveLinkHub/Source/LiveLinkHub/Private/LiveLinkHubTicker.h
+++ b/Engine/Plugins/Animation/LiveLinkHub/Source/LiveLinkHub/Private/LiveLinkHubTicker.h
@@ -5,8 +5,8 @@
 #include "HAL/Runnable.h"
 #include "HAL/RunnableThread.h"
 #include "LiveLinkHub.h"
-#include "LiveLinkHubSettings.h"
 #include "Misc/Timespan.h"
+#include "Settings/LiveLinkHubSettings.h"
 
 /** Object used to tick LiveLinkHub outside of the game thread. */
 class FLiveLinkHubTicker : public FRunnable
diff --git a/Engine/Plugins/Animation/LiveLinkHub/Source/LiveLinkHub/Private/Recording/Implementations/LiveLinkUAssetRecorder.cpp b/Engine/Plugins/Animation/LiveLinkHub/Source/LiveLinkHub/Private/Recording/Implementations/LiveLinkUAssetRecorder.cpp
--- a/Engine/Plugins/Animation/LiveLinkHub/Source/LiveLinkHub/Private/Recording/Implementations/LiveLinkUAssetRecorder.cpp
+++ b/Engine/Plugins/Animation/LiveLinkHub/Source/LiveLinkHub/Private/Recording/Implementations/LiveLinkUAssetRecorder.cpp
@@ -2,31 +2,35 @@
 
 #include "LiveLinkUAssetRecorder.h"
 
-#include "Recording/LiveLinkRecording.h"
-#include "Recording/Implementations/LiveLinkUAssetRecording.h"
-
+#include "AssetRegistry/AssetRegistryModule.h"
 #include "AssetToolsModule.h"
+#include "Config/LiveLinkHubFileUtilities.h"
 #include "Containers/Map.h"
 #include "ContentBrowserModule.h"
 #include "Features/IModularFeatures.h"
-#include "Framework/MultiBox/MultiBoxBuilder.h"
 #include "FileHelpers.h"
-#include "StructUtils/InstancedStruct.h"
+#include "Framework/MultiBox/MultiBoxBuilder.h"
+#include "HAL/PlatformFileManager.h"
 #include "IContentBrowserSingleton.h"
 #include "LiveLinkHub.h"
 #include "LiveLinkHubClient.h"
+#include "LiveLinkHubLog.h"
 #include "LiveLinkHubModule.h"
-#include "LiveLinkTypes.h"
 #include "LiveLinkPreset.h"
+#include "LiveLinkTypes.h"
 #include "Misc/FileHelper.h"
 #include "Misc/PackageName.h"
+#include "Recording/Implementations/LiveLinkUAssetRecording.h"
+#include "Recording/LiveLinkRecording.h"
+#include "Settings/LiveLinkHubSettings.h"
+#include "StructUtils/InstancedStruct.h"
 #include "UI/Window/LiveLinkHubWindowController.h"
-#include "Widgets/DeclarativeSyntaxSupport.h"
-#include "Widgets/Layout/SBox.h"
-#include "Widgets/SWidget.h"
 #include "UObject/Object.h"
 #include "UObject/Package.h"
 #include "UObject/StrongObjectPtr.h"
+#include "Widgets/DeclarativeSyntaxSupport.h"
+#include "Widgets/Layout/SBox.h"
+#include "Widgets/SWidget.h"
 
 #define LOCTEXT_NAMESPACE "LiveLinkHub.RecordingController"
 
@@ -144,22 +148,43 @@ bool FLiveLinkUAssetRecorder::OpenSaveDialog(const FString& InDefaultPath, const
 
 bool FLiveLinkUAssetRecorder::GetSavePresetPackageName(FString& OutName)
 {
-	FDateTime Today = FDateTime::Now();
+	using namespace UE::LiveLinkHub::FileUtilities::Private;
+	const FDateTime Today = FDateTime::Now();
 
 	TMap<FString, FStringFormatArg> FormatArgs;
 	FormatArgs.Add(TEXT("date"), Today.ToString());
+	
+	FFilenameTemplateData TemplateData;
+	ParseFilenameTemplate(GetDefault<ULiveLinkHubSettings>()->FilenameTemplate, TemplateData);
 
-	FString DialogStartPath = TEXT("/Game/LiveLinkRecordings");
+	const FString DefaultName = TemplateData.FileName;
+	const FString DefaultFolder = TemplateData.FolderPath;
+	
+	const FString ContentDir = FPaths::ProjectContentDir();
+	const FString DialogStartPath = FPaths::Combine(TEXT("/Game"), TemplateData.FolderPath);
+	const FString AbsoluteFolderPath = FPaths::Combine(ContentDir, TemplateData.FolderPath);
 
-	FString DefaultName = LOCTEXT("NewLiveLinkRecordingName", "NewLiveLinkRecording").ToString();
+	IPlatformFile& PlatformFile = FPlatformFileManager::Get().GetPlatformFile();
 
+	// Create directory if it doesn't exist
+	if (!PlatformFile.DirectoryExists(*AbsoluteFolderPath))
+	{
+		if (!PlatformFile.CreateDirectoryTree(*AbsoluteFolderPath))
+		{
+			UE_LOG(LogLiveLinkHub, Error, TEXT("Failed to create directory %s."), *AbsoluteFolderPath);
+			return false;
+		}
+
+		FAssetRegistryModule& AssetRegistryModule = FModuleManager::LoadModuleChecked<FAssetRegistryModule>("AssetRegistry");
+		AssetRegistryModule.Get().ScanPathsSynchronous({ TEXT("/Game") }, true);
+	}
 	FString UniquePackageName;
 	FString UniqueAssetName;
 
 	FAssetToolsModule& AssetToolsModule = FModuleManager::LoadModuleChecked<FAssetToolsModule>("AssetTools");
-	AssetToolsModule.Get().CreateUniqueAssetName(DialogStartPath / DefaultName, TEXT(""), UniquePackageName, UniqueAssetName);
+	AssetToolsModule.Get().CreateUniqueAssetName(DefaultFolder / DefaultName, TEXT(""), UniquePackageName, UniqueAssetName);
 
-	FString DialogStartName = FPaths::GetCleanFilename(UniqueAssetName);
+	const FString DialogStartName = FPaths::GetCleanFilename(UniqueAssetName);
 
 	FString UserPackageName;
 	FString NewPackageName;
diff --git a/Engine/Plugins/Animation/LiveLinkHub/Source/LiveLinkHub/Private/Recording/Implementations/LiveLinkUAssetRecording.cpp b/Engine/Plugins/Animation/LiveLinkHub/Source/LiveLinkHub/Private/Recording/Implementations/LiveLinkUAssetRecording.cpp
--- a/Engine/Plugins/Animation/LiveLinkHub/Source/LiveLinkHub/Private/Recording/Implementations/LiveLinkUAssetRecording.cpp
+++ b/Engine/Plugins/Animation/LiveLinkHub/Source/LiveLinkHub/Private/Recording/Implementations/LiveLinkUAssetRecording.cpp
@@ -635,12 +635,13 @@ void ULiveLinkUAssetRecording::LoadFrameData(FFrameFileData& InFrameData, FLiveL
 
 FString ULiveLinkUAssetRecording::GetRecordingDataFilePath() const
 {
-	// todo: This path needs to be calculated from the settings file.
-	FString AssetPath = GetPathName();
-	FString BasePath = FPaths::ProjectSavedDir();
-	FString FilePath = FString::Printf(TEXT("%sRecordingData/%s.rec"), *BasePath, *FPaths::GetBaseFilename(AssetPath));
+	const FString AssetPath = GetPathName();
+	FString ObjectDirectory = FPaths::GetPath(AssetPath);
+	ObjectDirectory.RemoveFromStart(TEXT("/Game"));
+	const FString AbsoluteFolderPath = FPaths::ConvertRelativePathToFull(FPaths::Combine(FPaths::ProjectContentDir(), ObjectDirectory));
+	const FString AbsoluteFilePath = FString::Printf(TEXT("%s/%s.rec"), *AbsoluteFolderPath, *FPaths::GetBaseFilename(AssetPath));
 
-	return FilePath;
+	return AbsoluteFilePath;
 }
 
 void ULiveLinkUAssetRecording::FLiveLinkStreamAsyncTask::DoWork()
diff --git a/Engine/Plugins/Animation/LiveLinkHub/Source/LiveLinkHub/Private/Recording/Implementations/LiveLinkUAssetRecordingPlayer.cpp b/Engine/Plugins/Animation/LiveLinkHub/Source/LiveLinkHub/Private/Recording/Implementations/LiveLinkUAssetRecordingPlayer.cpp
--- a/Engine/Plugins/Animation/LiveLinkHub/Source/LiveLinkHub/Private/Recording/Implementations/LiveLinkUAssetRecordingPlayer.cpp
+++ b/Engine/Plugins/Animation/LiveLinkHub/Source/LiveLinkHub/Private/Recording/Implementations/LiveLinkUAssetRecordingPlayer.cpp
@@ -4,8 +4,8 @@
 
 #include "HAL/IConsoleManager.h"
 #include "LiveLinkHubLog.h"
-#include "LiveLinkHubSettings.h"
 #include "Recording/LiveLinkRecording.h"
+#include "Settings/LiveLinkHubSettings.h"
 
 
 class FLiveLinkPlaybackTrackIterator
diff --git a/Engine/Plugins/Animation/LiveLinkHub/Source/LiveLinkHub/Private/Settings/LiveLinkHubSettings.cpp b/Engine/Plugins/Animation/LiveLinkHub/Source/LiveLinkHub/Private/Settings/LiveLinkHubSettings.cpp
new file mode 100644
--- /dev/null
+++ b/Engine/Plugins/Animation/LiveLinkHub/Source/LiveLinkHub/Private/Settings/LiveLinkHubSettings.cpp
@@ -0,0 +1,23 @@
+// Copyright Epic Games, Inc. All Rights Reserved.
+
+#include "LiveLinkHubSettings.h"
+
+#include "Config/LiveLinkHubFileUtilities.h"
+#include "Config/LiveLinkHubTemplateTokens.h"
+
+void ULiveLinkHubSettings::PostEditChangeProperty(FPropertyChangedEvent& PropertyChangedEvent)
+{
+	UObject::PostEditChangeProperty(PropertyChangedEvent);
+	if (PropertyChangedEvent.GetPropertyName() == GET_MEMBER_NAME_CHECKED(ULiveLinkHubSettings, FilenameTemplate))
+	{
+		CalculateExampleOutput();
+	}
+}
+
+void ULiveLinkHubSettings::CalculateExampleOutput()
+{
+	using namespace UE::LiveLinkHub::FileUtilities::Private;
+	FFilenameTemplateData TemplateData;
+	ParseFilenameTemplate(FilenameTemplate, TemplateData);
+	FilenameOutput = TemplateData.FullPath;
+}
diff --git a/Engine/Plugins/Animation/LiveLinkHub/Source/LiveLinkHub/Private/Settings/LiveLinkHubSettings.h b/Engine/Plugins/Animation/LiveLinkHub/Source/LiveLinkHub/Private/Settings/LiveLinkHubSettings.h
new file mode 100644
--- /dev/null
+++ b/Engine/Plugins/Animation/LiveLinkHub/Source/LiveLinkHub/Private/Settings/LiveLinkHubSettings.h
@@ -0,0 +1,64 @@
+// Copyright Epic Games, Inc. All Rights Reserved.
+
+#pragma once
+
+#include "UObject/Object.h"
+
+#include "Misc/FrameRate.h"
+
+#include "LiveLinkHubSettings.generated.h"
+
+/**
+ * Settings for LiveLinkHub.
+ */
+UCLASS(config=Engine, defaultconfig)
+class LIVELINKHUB_API ULiveLinkHubSettings : public UObject
+{
+	GENERATED_BODY()
+
+public:
+	/** Parse templates and set example output fields. */
+	void CalculateExampleOutput();
+	
+protected:
+	//~ Begin UObject interface
+	virtual void PostEditChangeProperty(FPropertyChangedEvent& PropertyChangedEvent) override;
+	//~ End UObject interface
+	
+public:
+	/** If enabled, discovered clients will be automatically added to the current session. */
+	UPROPERTY(config, EditAnywhere, Category="LiveLinkHub", DisplayName = "Automatically add discovered clients")
+	bool bAutoAddDiscoveredClients = true;
+
+	/** The size in megabytes to buffer when streaming a recoding. */
+	UPROPERTY(config, EditAnywhere, Category="LiveLinkHub", meta = (ClampMin = "1", UIMin = "1"))
+	int32 FrameBufferSizeMB = 100;
+	
+	/** Which project settings sections to display when opening the settings viewer. */
+	UPROPERTY(config)
+	TArray<FName> ProjectSettingsToDisplay;
+
+	/**
+	 * - Experimental - If this is disabled, LiveLinkHub's LiveLink Client will tick outside of the game thread.
+	 * This allows processing LiveLink frame snapshots without the risk of being blocked by the game / ui thread.
+	 * Note that this should only be relevant for virtual subjects since data is already forwarded to UE outside of the game thread.
+	 */
+	UPROPERTY(config, EditAnywhere, Category = "LiveLinkHub", meta = (ConfigRestartRequired = true))
+	bool bTickOnGameThread = false;
+
+	/** Target framerate for ticking LiveLinkHub. */
+	UPROPERTY(config, EditAnywhere, Category="LiveLinkHub", meta = (ConfigRestartRequired = true, ClampMin="15.0"))
+	float TargetFrameRate = 60.0f;
+
+	/** The filename template to use when creating recordings. */
+	UPROPERTY(config, EditAnywhere, Category="Templates")
+	FString FilenameTemplate = TEXT("NewLiveLinkRecording");
+
+	/** Example parsed output of the template. */
+	UPROPERTY(VisibleAnywhere, Category="Templates", DisplayName="Output")
+	FString FilenameOutput;
+
+	/** Placeholder for a list of the automatic tokens, set from the customization. */
+	UPROPERTY(VisibleAnywhere, Category="Templates")
+	FText AutomaticTokens;
+};
diff --git a/Engine/Plugins/Animation/LiveLinkHub/Source/LiveLinkHub/Private/Settings/LiveLinkHubSettingsCustomization.cpp b/Engine/Plugins/Animation/LiveLinkHub/Source/LiveLinkHub/Private/Settings/LiveLinkHubSettingsCustomization.cpp
new file mode 100644
--- /dev/null
+++ b/Engine/Plugins/Animation/LiveLinkHub/Source/LiveLinkHub/Private/Settings/LiveLinkHubSettingsCustomization.cpp
@@ -0,0 +1,71 @@
+// Copyright Epic Games, Inc. All Rights Reserved.
+
+#include "LiveLinkHubSettingsCustomization.h"
+
+#include "Config/LiveLinkHubTemplateTokens.h"
+#include "DetailCategoryBuilder.h"
+#include "DetailLayoutBuilder.h"
+#include "DetailWidgetRow.h"
+#include "LiveLinkHubSettings.h"
+#include "Widgets/Text/STextBlock.h"
+
+#define LOCTEXT_NAMESPACE "LiveLinkHubSettingsCustomization"
+
+TSharedRef<IDetailCustomization> FLiveLinkHubSettingsCustomization::MakeInstance()
+{
+	return MakeShared<FLiveLinkHubSettingsCustomization>();
+}
+
+void FLiveLinkHubSettingsCustomization::CustomizeDetails(IDetailLayoutBuilder& DetailBuilder)
+{
+	// Update current value when opening the settings page. Not safe to do under the settings object since there isn't
+	// an explicit callback when the settings page is opened, and using a method like PostInitProperties fires on the CDO
+	// early in the startup process.
+	GetMutableDefault<ULiveLinkHubSettings>()->CalculateExampleOutput();
+	
+	const TSharedRef<IPropertyHandle> AutomaticTokensHandle = DetailBuilder.GetProperty(GET_MEMBER_NAME_CHECKED(ULiveLinkHubSettings, AutomaticTokens));
+	IDetailPropertyRow* AutomaticTokensRow = DetailBuilder.EditDefaultProperty(AutomaticTokensHandle);
+	check(AutomaticTokensRow);
+	
+	const FLiveLinkHubAutomaticTokens& AutomaticTokens = FLiveLinkHubAutomaticTokens::GetStaticTokens();
+	FString FormattedTokensString;
+	
+	for (TFieldIterator<FProperty> It(FLiveLinkHubAutomaticTokens::StaticStruct()); It; ++It)
+	{
+		FProperty* Property = *It;
+		if (const FStrProperty* StrProperty = CastField<FStrProperty>(Property))
+		{
+			const FString PropertyValue = StrProperty->GetPropertyValue_InContainer(&AutomaticTokens);
+			const FString Token = UE::LiveLinkHub::Tokens::Private::CreateToken(PropertyValue);
+
+			const FString Tooltip = Property->GetToolTipText().ToString();
+			if (!Tooltip.IsEmpty())
+			{
+				FormattedTokensString += FString::Printf(TEXT("%s - %s\n"), *Token, *Tooltip);
+			}
+		}
+	}
+
+	AutomaticTokensRow->CustomWidget()
+	.WholeRowContent()
+	[
+		SNew(SVerticalBox)
+		+ SVerticalBox::Slot()
+		.Padding(0.f, 4.f)
+		.AutoHeight()
+		[
+			SNew(STextBlock)
+			.Text(AutomaticTokensHandle->GetPropertyDisplayName())
+			.Font(IDetailLayoutBuilder::GetDetailFontBold())
+		]
+		+ SVerticalBox::Slot()
+		.AutoHeight()
+		[
+			SNew(STextBlock)
+			.Text(FText::FromString(FormattedTokensString))
+			.Font(IDetailLayoutBuilder::GetDetailFont())
+		]
+	];
+}
+
+#undef LOCTEXT_NAMESPACE
\ No newline at end of file
diff --git a/Engine/Plugins/Animation/LiveLinkHub/Source/LiveLinkHub/Private/Settings/LiveLinkHubSettingsCustomization.h b/Engine/Plugins/Animation/LiveLinkHub/Source/LiveLinkHub/Private/Settings/LiveLinkHubSettingsCustomization.h
new file mode 100644
--- /dev/null
+++ b/Engine/Plugins/Animation/LiveLinkHub/Source/LiveLinkHub/Private/Settings/LiveLinkHubSettingsCustomization.h
@@ -0,0 +1,16 @@
+// Copyright Epic Games, Inc. All Rights Reserved.
+
+#pragma once
+
+#include "CoreMinimal.h"
+#include "IDetailCustomization.h"
+
+/**
+ * Customization for the live link hub settings.
+ */
+class FLiveLinkHubSettingsCustomization : public IDetailCustomization
+{
+public:
+	static TSharedRef<IDetailCustomization> MakeInstance();
+	virtual void CustomizeDetails(IDetailLayoutBuilder& DetailBuilder) override;
+};
\ No newline at end of file
diff --git a/Engine/Plugins/Animation/LiveLinkHub/Source/LiveLinkHub/Private/UI/Widgets/SLiveLinkHubSettings.h b/Engine/Plugins/Animation/LiveLinkHub/Source/LiveLinkHub/Private/UI/Widgets/SLiveLinkHubSettings.h
--- a/Engine/Plugins/Animation/LiveLinkHub/Source/LiveLinkHub/Private/UI/Widgets/SLiveLinkHubSettings.h
+++ b/Engine/Plugins/Animation/LiveLinkHub/Source/LiveLinkHub/Private/UI/Widgets/SLiveLinkHubSettings.h
@@ -6,14 +6,14 @@
 #include "Engine/DeveloperSettings.h"
 #include "Framework/Docking/TabManager.h"
 #include "Framework/MultiBox/MultiBoxBuilder.h"
-#include "Modules/ModuleManager.h"
 #include "ISettingsCategory.h"
 #include "ISettingsContainer.h"
 #include "ISettingsEditorModel.h"
 #include "ISettingsEditorModule.h"
 #include "ISettingsModule.h"
 #include "ISettingsSection.h"
-#include "LiveLinkHubSettings.h"
+#include "Modules/ModuleManager.h"
+#include "Settings/LiveLinkHubSettings.h"
 #include "Widgets/Docking/SDockTab.h"
 #include "Widgets/Input/SComboBox.h"
 #include "Widgets/SCompoundWidget.h"
