LiveLinkHub - Only write recording file data when saving the initial recording asset, rather than rely on serialize as recordings never update their animation data.

horde issue 507067

diff --git a/Engine/Plugins/Animation/LiveLinkHub/Source/LiveLinkHub/Private/Recording/Implementations/LiveLinkUAssetRecorder.cpp b/Engine/Plugins/Animation/LiveLinkHub/Source/LiveLinkHub/Private/Recording/Implementations/LiveLinkUAssetRecorder.cpp
--- a/Engine/Plugins/Animation/LiveLinkHub/Source/LiveLinkHub/Private/Recording/Implementations/LiveLinkUAssetRecorder.cpp
+++ b/Engine/Plugins/Animation/LiveLinkHub/Source/LiveLinkHub/Private/Recording/Implementations/LiveLinkUAssetRecorder.cpp
@@ -226,6 +226,10 @@ void FLiveLinkUAssetRecorder::SaveRecording()
 		NewRecording->LengthInSeconds = TimeRecordingEnded - TimeRecordingStarted;
 		NewRecording->RecordingPreset->BuildFromClient();
 		NewRecording->RecordingData = MoveTemp(*CurrentRecording);
+
+		// todo: async stream save and utilize bulk data within UAsset instead of a separate file.
+		NewRecording->SaveRecording();
+		
 		NewRecording->MarkPackageDirty();
 
 		//FAssetRegistryModule::AssetCreated(NewPreset);  Disabled for now since unsure if needed.
diff --git a/Engine/Plugins/Animation/LiveLinkHub/Source/LiveLinkHub/Private/Recording/Implementations/LiveLinkUAssetRecording.cpp b/Engine/Plugins/Animation/LiveLinkHub/Source/LiveLinkHub/Private/Recording/Implementations/LiveLinkUAssetRecording.cpp
--- a/Engine/Plugins/Animation/LiveLinkHub/Source/LiveLinkHub/Private/Recording/Implementations/LiveLinkUAssetRecording.cpp
+++ b/Engine/Plugins/Animation/LiveLinkHub/Source/LiveLinkHub/Private/Recording/Implementations/LiveLinkUAssetRecording.cpp
@@ -19,14 +19,6 @@ ULiveLinkUAssetRecording::~ULiveLinkUAssetRecording()
 void ULiveLinkUAssetRecording::Serialize(FArchive& Ar)
 {
 	Super::Serialize(Ar);
-	
-	FString FilePath = GetRecordingDataFilePath();
-	
-	if (Ar.IsSaving())
-	{
-		// todo: async stream save
-		SaveRecording();
-	}
 }
 
 void ULiveLinkUAssetRecording::SaveRecording()
