LiveLink Hub: Remove uses of WITH_LIVELINK_HUB preprocessor definition.

Several behaviors in common LiveLink code that were intended to be specific to LiveLink Hub were previously controlled by the WITH_LIVELINK_HUB preprocessor definition, which precluded sharing a build environment with the rest of the editor.

This CL replaces those preprocessor conditionals with appropriate config directives that are overridden in `Programs/LiveLinkHub/Config/DefaultEngine.ini`.

diff --git a/Engine/Plugins/Animation/LiveLink/Source/LiveLink/Private/LiveLinkClientReference.cpp b/Engine/Plugins/Animation/LiveLink/Source/LiveLink/Private/LiveLinkClientReference.cpp
--- a/Engine/Plugins/Animation/LiveLink/Source/LiveLink/Private/LiveLinkClientReference.cpp
+++ b/Engine/Plugins/Animation/LiveLink/Source/LiveLink/Private/LiveLinkClientReference.cpp
@@ -4,17 +4,21 @@
 #include "ILiveLinkClient.h"
 #include "LiveLinkModule.h"
 #include "Features/IModularFeatures.h"
+#include "Misc/ConfigCacheIni.h"
 
-#ifndef WITH_LIVELINK_HUB
-#define WITH_LIVELINK_HUB 0
-#endif
 
-ILiveLinkClient* FLiveLinkClientReference::GetClient()const
+ILiveLinkClient* FLiveLinkClientReference::GetClient() const
 {
-#if WITH_LIVELINK_HUB
 	// Compiling for LiveLinkHub should return the livelink client we registered as a modular feature.
-	return &IModularFeatures::Get().GetModularFeature<ILiveLinkClient>(ILiveLinkClient::ModularFeatureName);
-#else
-	return FLiveLinkModule::LiveLinkClient_AnyThread;
-#endif
+	static const bool bUseModularClientReference = GConfig->GetBoolOrDefault(
+		TEXT("LiveLink"), TEXT("bUseModularClientReference"), false, GEngineIni);
+
+	if (bUseModularClientReference)
+	{
+		return &IModularFeatures::Get().GetModularFeature<ILiveLinkClient>(ILiveLinkClient::ModularFeatureName);
+	}
+	else
+	{
+		return FLiveLinkModule::LiveLinkClient_AnyThread;
+	}
 }
diff --git a/Engine/Plugins/Animation/LiveLink/Source/LiveLink/Private/LiveLinkModule.cpp b/Engine/Plugins/Animation/LiveLink/Source/LiveLink/Private/LiveLinkModule.cpp
--- a/Engine/Plugins/Animation/LiveLink/Source/LiveLink/Private/LiveLinkModule.cpp
+++ b/Engine/Plugins/Animation/LiveLink/Source/LiveLink/Private/LiveLinkModule.cpp
@@ -11,6 +11,7 @@
 #include "LiveLinkMessageBusDiscoveryManager.h"
 #include "LiveLinkSettings.h"
 #include "Misc/CommandLine.h"
+#include "Misc/ConfigCacheIni.h"
 #include "Misc/CoreDelegates.h"
 #include "Styling/SlateStyle.h"
 #include "Styling/SlateStyleRegistry.h"
@@ -18,10 +19,6 @@
 LLM_DEFINE_TAG(LiveLink);
 #define LOCTEXT_NAMESPACE "LiveLinkModule"
 
-#ifndef WITH_LIVELINK_HUB
-#define WITH_LIVELINK_HUB 0
-#endif
-
 
 FLiveLinkClient* FLiveLinkModule::LiveLinkClient_AnyThread = nullptr;
 
@@ -42,10 +39,14 @@ void FLiveLinkModule::StartupModule()
 	FLiveLinkLogInstance::CreateInstance();
 	CreateStyle();
 
-#if !WITH_LIVELINK_HUB
-	FPlatformAtomics::InterlockedExchangePtr((void**)&LiveLinkClient_AnyThread, &LiveLinkClient);
-	IModularFeatures::Get().RegisterModularFeature(FLiveLinkClient::ModularFeatureName, &LiveLinkClient);
-#endif
+	const bool bUseModularClientReference = GConfig->GetBoolOrDefault(
+		TEXT("LiveLink"), TEXT("bUseModularClientReference"), false, GEngineIni);
+
+	if (!bUseModularClientReference)
+	{
+		FPlatformAtomics::InterlockedExchangePtr((void**)&LiveLinkClient_AnyThread, &LiveLinkClient);
+		IModularFeatures::Get().RegisterModularFeature(FLiveLinkClient::ModularFeatureName, &LiveLinkClient);
+	}
 
 	LiveLinkMotionController.RegisterController();
 
@@ -64,11 +65,14 @@ void FLiveLinkModule::ShutdownModule()
 #endif
 	LiveLinkMotionController.UnregisterController();
 
+	const bool bUseModularClientReference = GConfig->GetBoolOrDefault(
+		TEXT("LiveLink"), TEXT("bUseModularClientReference"), false, GEngineIni);
 
-#if !WITH_LIVELINK_HUB
-	IModularFeatures::Get().UnregisterModularFeature(FLiveLinkClient::ModularFeatureName, &LiveLinkClient);
-	FPlatformAtomics::InterlockedExchangePtr((void**)&LiveLinkClient_AnyThread, nullptr);
-#endif
+	if (!bUseModularClientReference)
+	{
+		IModularFeatures::Get().UnregisterModularFeature(FLiveLinkClient::ModularFeatureName, &LiveLinkClient);
+		FPlatformAtomics::InterlockedExchangePtr((void**)&LiveLinkClient_AnyThread, nullptr);
+	}
 
 	FSlateStyleRegistry::UnRegisterSlateStyle(*StyleSet.Get());
 	FLiveLinkLogInstance::DestroyInstance();
diff --git a/Engine/Plugins/Animation/LiveLink/Source/LiveLinkEditor/Private/LiveLinkClientPanelToolbar.cpp b/Engine/Plugins/Animation/LiveLink/Source/LiveLinkEditor/Private/LiveLinkClientPanelToolbar.cpp
--- a/Engine/Plugins/Animation/LiveLink/Source/LiveLinkEditor/Private/LiveLinkClientPanelToolbar.cpp
+++ b/Engine/Plugins/Animation/LiveLink/Source/LiveLinkEditor/Private/LiveLinkClientPanelToolbar.cpp
@@ -22,6 +22,7 @@
 #include "LiveLinkSourceFactory.h"
 #include "Logging/MessageLog.h"
 #include "LiveLinkSourceSettings.h"
+#include "Misc/ConfigCacheIni.h"
 #include "Misc/FileHelper.h"
 #include "LiveLinkVirtualSubject.h"
 #include "Misc/MessageDialog.h"
@@ -44,10 +45,6 @@
 #include "Widgets/SNullWidget.h"
 #include "Widgets/SWindow.h"
 
-#ifndef WITH_LIVELINK_HUB
-#define WITH_LIVELINK_HUB 0
-#endif
-
 
 #define LOCTEXT_NAMESPACE "LiveLinkClientPanel"
 
diff --git a/Engine/Plugins/Animation/LiveLink/Source/LiveLinkEditor/Private/LiveLinkPanelController.cpp b/Engine/Plugins/Animation/LiveLink/Source/LiveLinkEditor/Private/LiveLinkPanelController.cpp
--- a/Engine/Plugins/Animation/LiveLink/Source/LiveLinkEditor/Private/LiveLinkPanelController.cpp
+++ b/Engine/Plugins/Animation/LiveLink/Source/LiveLinkEditor/Private/LiveLinkPanelController.cpp
@@ -8,14 +8,14 @@
 #include "LiveLinkClient.h"
 #include "LiveLinkClientCommands.h"
 #include "LiveLinkClientPanelViews.h"
+#include "Misc/ConfigCacheIni.h"
 #include "SLiveLinkDataView.h"
 
-#ifndef WITH_LIVELINK_HUB
-#define WITH_LIVELINK_HUB 0
-#endif
 
 FLiveLinkPanelController::FLiveLinkPanelController(TAttribute<bool> bInReadOnly)
 {
+	bSeparateSourcesSubjects = GConfig->GetBoolOrDefault(TEXT("LiveLink"), TEXT("bPanelControllerSeparateSourcesSubjects"), false, GEngineIni);
+
 	Client = (FLiveLinkClient*)&IModularFeatures::Get().GetModularFeature<ILiveLinkClient>(ILiveLinkClient::ModularFeatureName);
 
 	OnSourcesChangedHandle = Client->OnLiveLinkSourcesChanged().AddRaw(this, &FLiveLinkPanelController::OnSourcesChangedHandler);
@@ -100,15 +100,17 @@ void FLiveLinkPanelController::OnSubjectSelectionChangedHandler(FLiveLinkSubject
 		if (SubjectEntry->IsSource())
 		{
 			SourcesDetailsView->SetObject(SubjectEntry->GetSettings());
-#if !WITH_LIVELINK_HUB // Sources and subjects live in different tabs in the hub
-			SubjectsDetailsView->SetSubjectKey(FLiveLinkSubjectKey());
-#endif
+			if (!bSeparateSourcesSubjects)
+			{
+				SubjectsDetailsView->SetSubjectKey(FLiveLinkSubjectKey());
+			}
 		}
 		else
 		{
-#if !WITH_LIVELINK_HUB // Sources and subjects live in different tabs in the hub
-			SourcesDetailsView->SetObject(nullptr);
-#endif
+			if (!bSeparateSourcesSubjects)
+			{
+				SourcesDetailsView->SetObject(nullptr);
+			}
 			SubjectsDetailsView->SetSubjectKey(SubjectEntry->SubjectKey);
 		}
 		bDetailViewSet = true;
@@ -118,9 +120,10 @@ void FLiveLinkPanelController::OnSubjectSelectionChangedHandler(FLiveLinkSubject
 
 	if (!bDetailViewSet)
 	{
-#if !WITH_LIVELINK_HUB
-		SourcesDetailsView->SetObject(nullptr);
-#endif
+		if (!bSeparateSourcesSubjects)
+		{
+			SourcesDetailsView->SetObject(nullptr);
+		}
 		SubjectsDetailsView->SetSubjectKey(FLiveLinkSubjectKey());
 	}
 
@@ -129,10 +132,11 @@ void FLiveLinkPanelController::OnSubjectSelectionChangedHandler(FLiveLinkSubject
 	{
 		SourcesView->SourcesListView->SetSelection(SourcesView->SourceData[FoundSourceIndex]);
 
-#if WITH_LIVELINK_HUB
-		// Update source details to the selected subject's source.
-		SourcesDetailsView->SetObject(SourcesView->SourceData[FoundSourceIndex]->GetSourceSettings());
-#endif
+		if (bSeparateSourcesSubjects)
+		{
+			// Update source details to the selected subject's source.
+			SourcesDetailsView->SetObject(SourcesView->SourceData[FoundSourceIndex]->GetSourceSettings());
+		}
 	}
 	else
 	{
diff --git a/Engine/Plugins/Animation/LiveLink/Source/LiveLinkEditor/Public/LiveLinkPanelController.h b/Engine/Plugins/Animation/LiveLink/Source/LiveLinkEditor/Public/LiveLinkPanelController.h
--- a/Engine/Plugins/Animation/LiveLink/Source/LiveLinkEditor/Public/LiveLinkPanelController.h
+++ b/Engine/Plugins/Animation/LiveLink/Source/LiveLinkEditor/Public/LiveLinkPanelController.h
@@ -66,6 +66,9 @@ private:
 	// Hadnles subject selection changing.
 	void OnSubjectSelectionChangedHandler(FLiveLinkSubjectUIEntryPtr SubjectEntry, ESelectInfo::Type SelectInfo);
 
+	// Sources and subjects live in different tabs in LiveLink Hub.
+	bool bSeparateSourcesSubjects = false;
+
 public:
 	// Sources view
 	TSharedPtr<FLiveLinkSourcesView> SourcesView;
diff --git a/Engine/Plugins/Animation/LiveLinkHub/Source/LiveLinkHub/Private/LiveLinkHubModule.cpp b/Engine/Plugins/Animation/LiveLinkHub/Source/LiveLinkHub/Private/LiveLinkHubModule.cpp
--- a/Engine/Plugins/Animation/LiveLinkHub/Source/LiveLinkHub/Private/LiveLinkHubModule.cpp
+++ b/Engine/Plugins/Animation/LiveLinkHub/Source/LiveLinkHub/Private/LiveLinkHubModule.cpp
@@ -6,6 +6,7 @@
 #include "LiveLinkHubApplication.h"
 #include "LiveLinkHubLog.h"
 #include "LiveLinkHubSubjectSettings.h"
+#include "Misc/ConfigCacheIni.h"
 #include "Modules/ModuleManager.h"
 #include "PropertyEditorModule.h"
 #include "Recording/LiveLinkHubPlaybackController.h"
@@ -13,12 +14,9 @@
 #include "Settings/LiveLinkHubSettings.h"
 #include "Settings/LiveLinkHubSettingsCustomization.h"
 #include "Subjects/LiveLinkHubSubjectSettingsDetailsCustomization.h"
-
-#if !WITH_LIVELINK_HUB
 #include "HAL/FileManager.h"
 #include "Misc/AsyncTaskNotification.h"
 #include "ToolMenus.h"
-#endif
 
 #define LOCTEXT_NAMESPACE "LiveLinkHubModule"
 
@@ -37,7 +35,6 @@ void FLiveLinkHubModule::StartLiveLinkHub()
 
 void FLiveLinkHubModule::StartupModule()
 {
-#if !WITH_LIVELINK_HUB  // When running in the editor
 	FToolMenuOwnerScoped OwnerScoped(this);
 	UToolMenu* Menu = UToolMenus::Get()->ExtendMenu("LevelEditor.MainMenu.Tools");
 	FToolMenuSection& Section = Menu->AddSection("VirtualProductionSection", LOCTEXT("VirtualProductionSection", "Virtual Production"));
@@ -47,29 +44,31 @@ void FLiveLinkHubModule::StartupModule()
 		LOCTEXT("LiveLinkHubTooltip", "Launch the LiveLink Hub app."),
 		FSlateIcon("LiveLinkStyle", "LiveLinkClient.Common.Icon.Small"),
 		FUIAction(FExecuteAction::CreateRaw(this, &FLiveLinkHubModule::OpenLiveLinkHub)));
-#else
+
 	FPropertyEditorModule& PropertyModule = FModuleManager::LoadModuleChecked<FPropertyEditorModule>("PropertyEditor");
-	PropertyModule.RegisterCustomClassLayout(ULiveLinkHubSubjectSettings::StaticClass()->GetFName(), FOnGetDetailCustomizationInstance::CreateStatic(&FLiveLinkHubSubjectSettingsDetailsCustomization::MakeInstance));
-#endif
-	if (FPropertyEditorModule* PropertyEditorModule = FModuleManager::GetModulePtr<FPropertyEditorModule>("PropertyEditor"))
+
+	PropertyModule.RegisterCustomClassLayout(ULiveLinkHubSettings::StaticClass()->GetFName(), FOnGetDetailCustomizationInstance::CreateStatic(&FLiveLinkHubSettingsCustomization::MakeInstance));
+
+	bUseSubjectSettingsDetailsCustomization =
+		GConfig->GetBoolOrDefault(TEXT("LiveLink"), TEXT("bUseLiveLinkHubSubjectSettingsDetailsCustomization"), false, GEngineIni);
+	if (bUseSubjectSettingsDetailsCustomization)
 	{
-		PropertyEditorModule->RegisterCustomClassLayout(ULiveLinkHubSettings::StaticClass()->GetFName(), FOnGetDetailCustomizationInstance::CreateStatic(&FLiveLinkHubSettingsCustomization::MakeInstance));
+		PropertyModule.RegisterCustomClassLayout(ULiveLinkHubSubjectSettings::StaticClass()->GetFName(), FOnGetDetailCustomizationInstance::CreateStatic(&FLiveLinkHubSubjectSettingsDetailsCustomization::MakeInstance));
 	}
 }
 
 void FLiveLinkHubModule::ShutdownModule()
 {
-#if !WITH_LIVELINK_HUB
 	UToolMenus::UnregisterOwner(this);
-#else
-	if (FPropertyEditorModule* PropertyModule = FModuleManager::GetModulePtr<FPropertyEditorModule>("PropertyEditor"))
-	{
-		PropertyModule->UnregisterCustomClassLayout(ULiveLinkHubSubjectSettings::StaticClass()->GetFName());
-	}
-#endif
+
 	if (FPropertyEditorModule* PropertyEditorModule = FModuleManager::GetModulePtr<FPropertyEditorModule>("PropertyEditor"))
 	{
 		PropertyEditorModule->UnregisterCustomClassLayout(ULiveLinkHubSettings::StaticClass()->GetFName());
+
+		if (bUseSubjectSettingsDetailsCustomization)
+		{
+			PropertyEditorModule->UnregisterCustomClassLayout(ULiveLinkHubSubjectSettings::StaticClass()->GetFName());
+		}
 	}
 }
 
@@ -98,7 +97,6 @@ TSharedPtr<FLiveLinkHubPlaybackController> FLiveLinkHubModule::GetPlaybackContro
 	return LiveLinkHub ? LiveLinkHub->PlaybackController : nullptr;
 }
 
-#if !WITH_LIVELINK_HUB
 void FLiveLinkHubModule::OpenLiveLinkHub() const
 {
 	FAsyncTaskNotificationConfig NotificationConfig;
@@ -167,7 +165,6 @@ void FLiveLinkHubModule::OpenLiveLinkHub() const
 			false);
 	}
 }
-#endif /** WITH_LIVELINK_HUB */
 
 TSharedPtr<FLiveLinkHubSubjectController> FLiveLinkHubModule::GetSubjectController() const
 {
diff --git a/Engine/Plugins/Animation/LiveLinkHub/Source/LiveLinkHub/Private/LiveLinkHubModule.h b/Engine/Plugins/Animation/LiveLinkHub/Source/LiveLinkHub/Private/LiveLinkHubModule.h
--- a/Engine/Plugins/Animation/LiveLinkHub/Source/LiveLinkHub/Private/LiveLinkHubModule.h
+++ b/Engine/Plugins/Animation/LiveLinkHub/Source/LiveLinkHub/Private/LiveLinkHubModule.h
@@ -14,9 +14,6 @@ class FLiveLinkHubSubjectController;
 class FUICommandList;
 class ILiveLinkHubSessionManager;
 
-#ifndef WITH_LIVELINK_HUB
-#define WITH_LIVELINK_HUB 0
-#endif
 
 class FLiveLinkHubModule : public ILiveLinkHubModule
 {
@@ -42,11 +39,13 @@ public:
 	/** Get the subject controller. */
     TSharedPtr<ILiveLinkHubSessionManager> GetSessionManager() const;
 
-#if !WITH_LIVELINK_HUB
 	/** Launch livelink hub. */
 	void OpenLiveLinkHub() const;
-#endif
+
 private:
 	/** LiveLinkHub object responsible for initializing the different controllers. */
 	TSharedPtr<FLiveLinkHub> LiveLinkHub;
+
+	/** Config-driven. True in LiveLink Hub standalone application, false in UE. */
+	bool bUseSubjectSettingsDetailsCustomization;
 };
diff --git a/Engine/Plugins/Animation/LiveLinkHub/Source/LiveLinkHubMessaging/Private/LiveLinkHubMessageBusSourceFactory.cpp b/Engine/Plugins/Animation/LiveLinkHub/Source/LiveLinkHubMessaging/Private/LiveLinkHubMessageBusSourceFactory.cpp
--- a/Engine/Plugins/Animation/LiveLinkHub/Source/LiveLinkHubMessaging/Private/LiveLinkHubMessageBusSourceFactory.cpp
+++ b/Engine/Plugins/Animation/LiveLinkHub/Source/LiveLinkHubMessaging/Private/LiveLinkHubMessageBusSourceFactory.cpp
@@ -1,18 +1,14 @@
 // Copyright Epic Games, Inc. All Rights Reserved.
 
 #include "LiveLinkHubMessageBusSourceFactory.h"
-
 #include "LiveLinkHubMessageBusSource.h"
-
+#include "Misc/ConfigCacheIni.h"
 #include "SLiveLinkMessageBusSourceFactory.h"
 
-#ifndef WITH_LIVELINK_HUB
-#define WITH_LIVELINK_HUB 0
-#endif
-
 
 #define LOCTEXT_NAMESPACE "LiveLinkHubMessageBusSourceFactory"
 
+
 FText ULiveLinkHubMessageBusSourceFactory::GetSourceDisplayName() const
 {
 	return LOCTEXT("SourceDisplayName", "LiveLink Hub");
@@ -40,11 +36,8 @@ TSharedPtr<FLiveLinkMessageBusSource> ULiveLinkHubMessageBusSourceFactory::MakeS
 
 bool ULiveLinkHubMessageBusSourceFactory::IsEnabled() const
 {
-#if WITH_LIVELINK_HUB
-return false;
-#else
-return true;
-#endif
+	return GConfig->GetBoolOrDefault(TEXT("LiveLinkHub"), TEXT("bEnableLLHMessageBusSourceFactory"), true, GEngineIni);
 }
 
+
 #undef LOCTEXT_NAMESPACE
diff --git a/Engine/Plugins/Animation/LiveLinkHub/Source/LiveLinkHubMessaging/Private/LiveLinkHubMessagingModule.cpp b/Engine/Plugins/Animation/LiveLinkHub/Source/LiveLinkHubMessaging/Private/LiveLinkHubMessagingModule.cpp
--- a/Engine/Plugins/Animation/LiveLinkHub/Source/LiveLinkHubMessaging/Private/LiveLinkHubMessagingModule.cpp
+++ b/Engine/Plugins/Animation/LiveLinkHub/Source/LiveLinkHubMessaging/Private/LiveLinkHubMessagingModule.cpp
@@ -6,30 +6,30 @@
 #include "LiveLinkHubConnectionManager.h"
 #include "LiveLinkHubMessageBusSourceFactory.h"
 #include "LiveLinkMessageBusSourceFactory.h"
+#include "Misc/ConfigCacheIni.h"
 #include "Modules/ModuleManager.h"
 
 #ifndef WITH_LIVELINK_DISCOVERY_MANAGER_THREAD
 #define WITH_LIVELINK_DISCOVERY_MANAGER_THREAD 1
 #endif
 
-#ifndef WITH_LIVELINK_HUB
-#define WITH_LIVELINK_HUB 0
-#endif
-
-#ifndef LIVELINKHUB_USE_CONNECTION_MANAGER
-#define LIVELINKHUB_USE_CONNECTION_MANAGER !WITH_LIVELINK_HUB && WITH_LIVELINK_DISCOVERY_MANAGER_THREAD
-#endif
-
 
 class FLiveLinkHubMessagingModule : public ILiveLinkHubMessagingModule
 {
-	// The connection manager is used to communicate with the hub, so we don't need it when we're running the hub itself.
 public:
 	//~ Begin IModuleInterface
 	virtual void StartupModule() override
 	{
-#if LIVELINKHUB_USE_CONNECTION_MANAGER
-		ConnectionManager = MakeUnique<FLiveLinkHubConnectionManager>();
+		// The connection manager is used to communicate with the hub,
+		// so we don't need it when we're running the hub itself.
+		bUseConnectionManager = GConfig->GetBoolOrDefault(
+			TEXT("LiveLinkHub"), TEXT("bUseMessagingConnectionManager"), true, GEngineIni);
+
+#if WITH_LIVELINK_DISCOVERY_MANAGER_THREAD
+		if (bUseConnectionManager)
+		{
+			ConnectionManager = MakeUnique<FLiveLinkHubConnectionManager>();
+		}
 #endif
 
 		SourceFilterDelegate = ILiveLinkModule::Get().RegisterMessageBusSourceFilter(FOnLiveLinkShouldDisplaySource::CreateRaw(this, &FLiveLinkHubMessagingModule::OnFilterMessageBusSource));
@@ -42,7 +42,7 @@ public:
 			LiveLinkModule->UnregisterMessageBusSourceFilter(SourceFilterDelegate);
 		}
 
-#if LIVELINKHUB_USE_CONNECTION_MANAGER
+#if WITH_LIVELINK_DISCOVERY_MANAGER_THREAD
 		ConnectionManager.Reset();
 #endif
 	}
@@ -55,11 +55,6 @@ public:
 	}
 
 private:
-#if LIVELINKHUB_USE_CONNECTION_MANAGER
-	/** Manages the connection to the live link hub. */
-	TUniquePtr<FLiveLinkHubConnectionManager> ConnectionManager;
-#endif
-
 	/** Filter invoked by the messagebus source factory to filter out sources in the creation panel. */
 	bool OnFilterMessageBusSource(UClass* FactoryClass, FProviderPollResultPtr PollResult) const
 	{
@@ -70,32 +65,42 @@ private:
 
 		const bool bIsLiveLinkHubProvider = PollResult->Annotations.FindRef(FLiveLinkHubMessageAnnotation::ProviderTypeAnnotation) == UE::LiveLinkHub::Private::LiveLinkHubProviderType;
 
-#if WITH_LIVELINK_HUB
-		// When running the hub
-		if (FactoryClass == ULiveLinkMessageBusSourceFactory::StaticClass())
-		{
-			// Don't show the livelink hub in the message bus source discovery. (Should be changed in the future if we want to allow hubs to speak to each others).
-			return !bIsLiveLinkHubProvider;
-		}
-#else
-		// When running in UE
-		if (FactoryClass == ULiveLinkHubMessageBusSourceFactory::StaticClass())
+		if (!bUseConnectionManager)
 		{
-			return bIsLiveLinkHubProvider;
+			// When running the hub
+			if (FactoryClass == ULiveLinkMessageBusSourceFactory::StaticClass())
+			{
+				// Don't show the livelink hub in the message bus source discovery. (Should be changed in the future if we want to allow hubs to speak to each others).
+				return !bIsLiveLinkHubProvider;
+			}
 		}
-		else if (FactoryClass == ULiveLinkMessageBusSourceFactory::StaticClass())
+		else
 		{
-			return !bIsLiveLinkHubProvider;
+			// When running in UE
+			if (FactoryClass == ULiveLinkHubMessageBusSourceFactory::StaticClass())
+			{
+				return bIsLiveLinkHubProvider;
+			}
+			else if (FactoryClass == ULiveLinkMessageBusSourceFactory::StaticClass())
+			{
+				return !bIsLiveLinkHubProvider;
+			}
 		}
-#endif
 
 		return true;
 	}
 
+	bool bUseConnectionManager;
+
+#if WITH_LIVELINK_DISCOVERY_MANAGER_THREAD
+	/** Manages the connection to the live link hub. */
+	TUniquePtr<FLiveLinkHubConnectionManager> ConnectionManager;
+#endif
+
 	/** Handle to the delegate used to filter message bus sources. */
 	FDelegateHandle SourceFilterDelegate;
 
-	/** Delegate called when the connection between a livelink hub and the editor is establshed. */
+	/** Delegate called when the connection between a livelink hub and the editor is established. */
 	FOnHubConnectionEstablished ConnectionEstablishedDelegate;
 };
 
diff --git a/Engine/Programs/LiveLinkHub/Config/DefaultEngine.ini b/Engine/Programs/LiveLinkHub/Config/DefaultEngine.ini
--- a/Engine/Programs/LiveLinkHub/Config/DefaultEngine.ini
+++ b/Engine/Programs/LiveLinkHub/Config/DefaultEngine.ini
@@ -1,28 +1,11 @@
-[URL]
-GameNameShort=LLH
-
-[Plugins]
-+ProgramEnabledPlugins="LiveLink"
-+ProgramEnabledPlugins="LiveLinkHub"
-+ProgramEnabledPlugins="LiveLinkLens"
-+ProgramEnabledPlugins="LiveLinkCamera"
-+ProgramEnabledPlugins="LensComponent"
-+ProgramEnabledPlugins="LiveLinkInputDevice"
-+ProgramEnabledPlugins="ContentBrowserAssetDataSource"
-+ProgramEnabledPlugins="PropertyAccessEditor"
-+ProgramEnabledPlugins="PythonScriptPlugin"
-+ProgramEnabledPlugins="CameraCalibrationCore"
-+ProgramEnabledPlugins="ProceduralMeshComponent"
-
-; Optional plugins, will not cause an error if disabled.
-+ProgramOptionalPlugins="AppleARKitFaceSupport"
-
-+ProgramOptionalPlugins="LiveLinkViconDataStream"
-+ProgramOptionalPlugins="Smartsuit"
-+ProgramOptionalPlugins="PoseAILiveLink"
-;+ProgramOptionalPlugins="LiveLinkMVNPlugin"
-+ProgramOptionalPlugins="MocopiLiveLink"
-;+ProgramOptionalPlugins="OptitrackLiveLink"
+[LiveLink]
+bUseLiveLinkHubSubjectSettingsDetailsCustomization=true
+bPanelControllerSeparateSourcesSubjects=true
+bUseModularClientReference=true
+
+[LiveLinkHub]
+bEnableLLHMessageBusSourceFactory=false
+bUseMessagingConnectionManager=false
 
 [/Script/MacTargetPlatform.XcodeProjectSettings]
 TemplateMacPlist=(FilePath="/Engine/Source/Programs/LiveLinkHub/Resources/Mac/Info.plist")
@@ -46,4 +29,4 @@ bEnableLiveLinkForFaceTracking=true
 bShareMaterialShaderCode=false
 
 [ConsoleVariables]
-DetailsPanel.Style.EnableCardLayout = true
\ No newline at end of file
+DetailsPanel.Style.EnableCardLayout=true
diff --git a/Engine/Source/Programs/LiveLinkHubEditor/Config/DefaultEngine.ini b/Engine/Source/Programs/LiveLinkHubEditor/Config/DefaultEngine.ini
new file mode 100644
--- /dev/null
+++ b/Engine/Source/Programs/LiveLinkHubEditor/Config/DefaultEngine.ini
@@ -0,0 +1,35 @@
+[LiveLink]
+bUseLiveLinkHubSubjectSettingsDetailsCustomization=true
+bPanelControllerSeparateSourcesSubjects=true
+bUseModularClientReference=true
+
+[LiveLinkHub]
+bEnableLLHMessageBusSourceFactory=false
+bUseMessagingConnectionManager=false
+
+[/Script/MacTargetPlatform.XcodeProjectSettings]
+TemplateMacPlist=(FilePath="/Engine/Source/Programs/LiveLinkHub/Resources/Mac/Info.plist")
+
+[/Script/LiveLinkInterface.LiveLinkDefaultSourceSettings]
+DefaultSourceFrameBufferSize=1
+
+[/Script/LiveLinkInterface.LiveLinkDefaultSubjectSettings]
+bRebroadcastSubjectsByDefault=true
+bAllowModifyingRebroadcast=false
+
+[/Script/LiveLinkHub.LiveLinkHubSettings]
+bAutoAddDiscoveredClients=true
++ProjectSettingsToDisplay=LiveLink
++ProjectSettingsToDisplay=LiveLinkHub
+
+[/Script/AppleARKit.AppleARKitSettings]
+bEnableLiveLinkForFaceTracking=true
+
+[/Script/UnrealEd.ProjectPackagingSettings]
+bShareMaterialShaderCode=false
+
+[ConsoleVariables]
+DetailsPanel.Style.EnableCardLayout=true
+
+[/Script/BuildSettings.BuildSettings]
+DefaultEditorTarget=LiveLinkHubEditor
