LiveLinkHub - Fix LiveLinkHubClient sending frame data twice

Fix renaming not working anymore
Fix disabling subjects in clients view not working
Add a GetRebroadcastName method to LiveLinkSubjectSettings to allow overriding the subject's name when rebroadcasting it.

diff --git a/Engine/Plugins/Animation/LiveLink/Source/LiveLink/Private/LiveLinkClient.cpp b/Engine/Plugins/Animation/LiveLink/Source/LiveLink/Private/LiveLinkClient.cpp
--- a/Engine/Plugins/Animation/LiveLink/Source/LiveLink/Private/LiveLinkClient.cpp
+++ b/Engine/Plugins/Animation/LiveLink/Source/LiveLink/Private/LiveLinkClient.cpp
@@ -282,7 +282,6 @@ void FLiveLinkClient::HandleSubjectRebroadcast(ILiveLinkSubject* InSubject, cons
 
 				if (bPreProcessRebroadcastedFrames)
 				{
-					FLiveLinkSubject* LiveSubject = static_cast<FLiveLinkSubject*>(InSubject);
 					InSubject->PreprocessFrame(FrameDataCopy);
 				}
 
@@ -306,14 +305,16 @@ void FLiveLinkClient::HandleSubjectRebroadcast(ILiveLinkSubject* InSubject, cons
 					}
 				}
 
+				const FName RebroadcastName = GetRebroadcastName(InSubject->GetSubjectKey());
+
 				if (!InSubject->HasStaticDataBeenRebroadcasted())
 				{
-					RebroadcastLiveLinkProvider->UpdateSubjectStaticData(InSubject->GetSubjectKey().SubjectName, SubjectRole, MoveTemp(StaticDataCopy));
+					RebroadcastLiveLinkProvider->UpdateSubjectStaticData(RebroadcastName, SubjectRole, MoveTemp(StaticDataCopy));
 					InSubject->SetStaticDataAsRebroadcasted(true);
 					RebroadcastedSubjects.Add(InSubject->GetSubjectKey());
 				}
 
-				RebroadcastLiveLinkProvider->UpdateSubjectFrameData(InSubject->GetSubjectKey().SubjectName, MoveTemp(FrameDataCopy));
+				RebroadcastLiveLinkProvider->UpdateSubjectFrameData(RebroadcastName, MoveTemp(FrameDataCopy));
 			}
 			else
 			{
@@ -337,9 +338,11 @@ void FLiveLinkClient::RemoveRebroadcastedSubject(FLiveLinkSubjectKey InSubjectKe
 {
 	if (RebroadcastLiveLinkProvider.IsValid())
 	{
+		const FName SubjectName = GetRebroadcastName(InSubjectKey);
+
 		if (RebroadcastedSubjects.Contains(InSubjectKey))
 		{
-			RebroadcastLiveLinkProvider->RemoveSubject(InSubjectKey.SubjectName);
+			RebroadcastLiveLinkProvider->RemoveSubject(SubjectName);
 			RebroadcastedSubjects.Remove(InSubjectKey);
 
 			if (RebroadcastedSubjects.Num() <= 0)
@@ -1517,6 +1520,16 @@ FLiveLinkSubjectTimeSyncData FLiveLinkClient::GetTimeSyncData(FLiveLinkSubjectNa
 	return FLiveLinkSubjectTimeSyncData();
 }
 
+FName FLiveLinkClient::GetRebroadcastName(const FLiveLinkSubjectKey& InSubjectKey) const
+{
+	if (ULiveLinkSubjectSettings* Settings = Cast<ULiveLinkSubjectSettings>(GetSubjectSettings(InSubjectKey)))
+	{
+		return Settings->GetRebroadcastName();
+	}
+
+	return InSubjectKey.SubjectName;
+}
+
 FText FLiveLinkClient::GetSourceType(FGuid InEntryGuid) const
 {
 	if (const FLiveLinkCollectionSourceItem* SourceItem = Collection->FindSource(InEntryGuid))
diff --git a/Engine/Plugins/Animation/LiveLink/Source/LiveLink/Public/LiveLinkClient.h b/Engine/Plugins/Animation/LiveLink/Source/LiveLink/Public/LiveLinkClient.h
--- a/Engine/Plugins/Animation/LiveLink/Source/LiveLink/Public/LiveLinkClient.h
+++ b/Engine/Plugins/Animation/LiveLink/Source/LiveLink/Public/LiveLinkClient.h
@@ -166,6 +166,9 @@ public:
 
 	FLiveLinkSubjectTimeSyncData GetTimeSyncData(FLiveLinkSubjectName SubjectName);
 
+	/** Get the rebroadcast name for a given subject. (Defauls to the subject's subject name, but can be overriden. */
+	FName GetRebroadcastName(const FLiveLinkSubjectKey& InSubjectKey) const;
+
 	UE_DEPRECATED(4.23, "FLiveLinkClient::GetSourceTypeForEntry is deprecated. Please use GetSourceType instead!")
 	FText GetSourceTypeForEntry(FGuid EntryGuid) const { return GetSourceType(EntryGuid); }
 	UE_DEPRECATED(4.23, "FLiveLinkClient::GetMachineNameForEntry is deprecated. Please use GetSourceMachineName instead!")
diff --git a/Engine/Plugins/Animation/LiveLinkHub/Source/LiveLinkHub/Private/Clients/LiveLinkHubClientsModel.h b/Engine/Plugins/Animation/LiveLinkHub/Source/LiveLinkHub/Private/Clients/LiveLinkHubClientsModel.h
--- a/Engine/Plugins/Animation/LiveLinkHub/Source/LiveLinkHub/Private/Clients/LiveLinkHubClientsModel.h
+++ b/Engine/Plugins/Animation/LiveLinkHub/Source/LiveLinkHub/Private/Clients/LiveLinkHubClientsModel.h
@@ -56,8 +56,8 @@ public:
 	virtual void SetClientEnabled(FLiveLinkHubClientId InClient, bool bInEnable) = 0;
 
 	/** Get whether a subject is enabled on a given client. */
-	virtual bool IsSubjectEnabled(FLiveLinkHubClientId InClient, const FLiveLinkSubjectKey& Subject) const = 0;
+	virtual bool IsSubjectEnabled(FLiveLinkHubClientId InClient, FName SubjectName) const = 0;
 
 	/** Set whether a subject should receive livelink data. */
-	virtual void SetSubjectEnabled(FLiveLinkHubClientId InClient, const FLiveLinkSubjectKey& Subject, bool bInEnable) = 0;
+	virtual void SetSubjectEnabled(FLiveLinkHubClientId InClient, FName SubjectName, bool bInEnable) = 0;
 };
diff --git a/Engine/Plugins/Animation/LiveLinkHub/Source/LiveLinkHub/Private/Clients/LiveLinkHubProvider.cpp b/Engine/Plugins/Animation/LiveLinkHub/Source/LiveLinkHub/Private/Clients/LiveLinkHubProvider.cpp
--- a/Engine/Plugins/Animation/LiveLinkHub/Source/LiveLinkHub/Private/Clients/LiveLinkHubProvider.cpp
+++ b/Engine/Plugins/Animation/LiveLinkHub/Source/LiveLinkHub/Private/Clients/LiveLinkHubProvider.cpp
@@ -5,7 +5,6 @@
 #include "Algo/Transform.h"
 #include "Async/Async.h"
 #include "Clients/LiveLinkHubClientsModel.h"
-#include "Clients/LiveLinkHubProvider.h"
 #include "Clients/LiveLinkHubUEClientInfo.h"
 #include "Containers/ObservableArray.h"
 #include "CoreMinimal.h"
@@ -403,28 +402,28 @@ void FLiveLinkHubProvider::SetClientEnabled(FLiveLinkHubClientId Client, bool bI
 	}
 }
 
-bool FLiveLinkHubProvider::IsSubjectEnabled(FLiveLinkHubClientId Client, const FLiveLinkSubjectKey& Subject) const
+bool FLiveLinkHubProvider::IsSubjectEnabled(FLiveLinkHubClientId Client, FName SubjectName) const
 {
 	FReadScopeLock Locker(ClientsMapLock);
 	if (const FLiveLinkHubUEClientInfo* ClientInfoPtr = ClientsMap.Find(Client))
 	{
-		return !ClientInfoPtr->DisabledSubjects.Contains(Subject.SubjectName);
+		return !ClientInfoPtr->DisabledSubjects.Contains(SubjectName);
 	}
 	return false;
 }
 
-void FLiveLinkHubProvider::SetSubjectEnabled(FLiveLinkHubClientId Client, const FLiveLinkSubjectKey& Subject, bool bInEnable)
+void FLiveLinkHubProvider::SetSubjectEnabled(FLiveLinkHubClientId Client, FName SubjectName, bool bInEnable)
 {
 	FWriteScopeLock Locker(ClientsMapLock);
 	if (FLiveLinkHubUEClientInfo* ClientInfoPtr = ClientsMap.Find(Client))
 	{
 		if (bInEnable)
 		{
-			ClientInfoPtr->DisabledSubjects.Remove(Subject.SubjectName);
+			ClientInfoPtr->DisabledSubjects.Remove(SubjectName);
 		}
 		else
 		{
-			ClientInfoPtr->DisabledSubjects.Add(Subject.SubjectName);
+			ClientInfoPtr->DisabledSubjects.Add(SubjectName);
 		}
 	}
 }
diff --git a/Engine/Plugins/Animation/LiveLinkHub/Source/LiveLinkHub/Private/Clients/LiveLinkHubProvider.h b/Engine/Plugins/Animation/LiveLinkHub/Source/LiveLinkHub/Private/Clients/LiveLinkHubProvider.h
--- a/Engine/Plugins/Animation/LiveLinkHub/Source/LiveLinkHub/Private/Clients/LiveLinkHubProvider.h
+++ b/Engine/Plugins/Animation/LiveLinkHub/Source/LiveLinkHub/Private/Clients/LiveLinkHubProvider.h
@@ -106,8 +106,8 @@ protected:
 	virtual bool IsClientEnabled(FLiveLinkHubClientId Client) const override;
 	virtual bool IsClientConnected(FLiveLinkHubClientId Client) const override;
 	virtual void SetClientEnabled(FLiveLinkHubClientId Client, bool bInEnable) override;
-	virtual bool IsSubjectEnabled(FLiveLinkHubClientId Client, const FLiveLinkSubjectKey& Subject) const override;
-	virtual void SetSubjectEnabled(FLiveLinkHubClientId Client, const FLiveLinkSubjectKey& Subject, bool bInEnable) override;
+	virtual bool IsSubjectEnabled(FLiveLinkHubClientId Client, FName SubjectName) const override;
+	virtual void SetSubjectEnabled(FLiveLinkHubClientId Client, FName SubjectName, bool bInEnable) override;
 	//~ End ILiveLinkHubClientsModel interface
 
 private:
diff --git a/Engine/Plugins/Animation/LiveLinkHub/Source/LiveLinkHub/Private/Clients/SLiveLinkHubClientsView.h b/Engine/Plugins/Animation/LiveLinkHub/Source/LiveLinkHub/Private/Clients/SLiveLinkHubClientsView.h
--- a/Engine/Plugins/Animation/LiveLinkHub/Source/LiveLinkHub/Private/Clients/SLiveLinkHubClientsView.h
+++ b/Engine/Plugins/Animation/LiveLinkHub/Source/LiveLinkHub/Private/Clients/SLiveLinkHubClientsView.h
@@ -163,7 +163,8 @@ struct FClientTreeViewSubjectItem : public FClientTreeViewItem
 	{
 		if (const TSharedPtr<ILiveLinkHubClientsModel> ClientsModelPtr = ClientsModel.Pin())
 		{
-			return ClientsModelPtr->IsSubjectEnabled(ClientId, LiveLinkSubjectKey);
+			FLiveLinkClient& LiveLinkClient = static_cast<FLiveLinkClient&>(IModularFeatures::Get().GetModularFeature<ILiveLinkClient>(ILiveLinkClient::ModularFeatureName));
+			return ClientsModelPtr->IsSubjectEnabled(ClientId, LiveLinkClient.GetRebroadcastName(LiveLinkSubjectKey));
 		}
 		return false;
 	}
@@ -182,7 +183,8 @@ struct FClientTreeViewSubjectItem : public FClientTreeViewItem
 	{
 		if (const TSharedPtr<ILiveLinkHubClientsModel> ClientsModelPtr = ClientsModel.Pin())
 		{
-			ClientsModelPtr->SetSubjectEnabled(ClientId, LiveLinkSubjectKey, bInEnabled);
+			FLiveLinkClient& LiveLinkClient = static_cast<FLiveLinkClient&>(IModularFeatures::Get().GetModularFeature<ILiveLinkClient>(ILiveLinkClient::ModularFeatureName));
+			ClientsModelPtr->SetSubjectEnabled(ClientId, LiveLinkClient.GetRebroadcastName(LiveLinkSubjectKey), bInEnabled);
 		}
 	}
 
diff --git a/Engine/Plugins/Animation/LiveLinkHub/Source/LiveLinkHub/Private/LiveLinkHub.cpp b/Engine/Plugins/Animation/LiveLinkHub/Source/LiveLinkHub/Private/LiveLinkHub.cpp
--- a/Engine/Plugins/Animation/LiveLinkHub/Source/LiveLinkHub/Private/LiveLinkHub.cpp
+++ b/Engine/Plugins/Animation/LiveLinkHub/Source/LiveLinkHub/Private/LiveLinkHub.cpp
@@ -149,14 +149,6 @@ void FLiveLinkHub::OnStaticDataReceived_AnyThread(const FLiveLinkSubjectKey& InS
 	{
 		RecordingController->RecordStaticData(InSubjectKey, InRole, InStaticDataStruct);
 	}
-
-	FLiveLinkStaticDataStruct StaticDataCopy;
-	StaticDataCopy.InitializeWith(InStaticDataStruct);
-
-	UE_LOG(LogLiveLinkHub, Verbose, TEXT("Pushing static data for %s"), *InSubjectKey.SubjectName.ToString());
-
-	const FName OverridenName = GetSubjectNameOverride(InSubjectKey);
-	LiveLinkProvider->UpdateSubjectStaticData(OverridenName, InRole, MoveTemp(StaticDataCopy));
 }
 
 void FLiveLinkHub::OnFrameDataReceived_AnyThread(const FLiveLinkSubjectKey& InSubjectKey, const FLiveLinkFrameDataStruct& InFrameDataStruct) const
@@ -165,15 +157,6 @@ void FLiveLinkHub::OnFrameDataReceived_AnyThread(const FLiveLinkSubjectKey& InSu
 	{
 		RecordingController->RecordFrameData(InSubjectKey, InFrameDataStruct);
 	}
-
-	FLiveLinkFrameDataStruct FrameDataCopy;
-	FrameDataCopy.InitializeWith(InFrameDataStruct);
-
-	const FName OverridenName = GetSubjectNameOverride(InSubjectKey);
-	if (LiveLinkHubClient->IsSubjectEnabled(InSubjectKey.SubjectName))
-	{
-		LiveLinkProvider->UpdateSubjectFrameData(OverridenName, MoveTemp(FrameDataCopy));
-	}
 }
 
 void FLiveLinkHub::OnSubjectMarkedPendingKill_AnyThread(const FLiveLinkSubjectKey& InSubjectKey) const
@@ -181,7 +164,7 @@ void FLiveLinkHub::OnSubjectMarkedPendingKill_AnyThread(const FLiveLinkSubjectKe
 	UE_LOG(LogLiveLinkHub, Verbose, TEXT("Removed subject %s"), *InSubjectKey.SubjectName.ToString());
 
 	// Send an update to connected clients as well.
-	const FName OverridenName = GetSubjectNameOverride(InSubjectKey);
+	const FName OverridenName = LiveLinkHubClient->GetRebroadcastName(InSubjectKey);
 
 	// Note: We send a RemoveSubject message to connected clients when the subject is marked pending kill in order to process this message in the right order.
 	// If we were to send a RemoveSubject message after the OnSubjectRemoved delegate, it could cause our RemoveSubject message to be sent out of order.
@@ -223,17 +206,6 @@ void FLiveLinkHub::OpenConfig()
 	SessionManager->RestoreSession();
 }
 
-FName FLiveLinkHub::GetSubjectNameOverride(const FLiveLinkSubjectKey& InSubjectKey) const
-{
-	ILiveLinkClient& Client = IModularFeatures::Get().GetModularFeature<ILiveLinkClient>(ILiveLinkClient::ModularFeatureName);
-	if (ULiveLinkHubSubjectSettings* Settings = Cast<ULiveLinkHubSubjectSettings>(Client.GetSubjectSettings(InSubjectKey)))
-	{
-		return *Settings->OutboundName;
-	}
-
-	return InSubjectKey.SubjectName;
-}
-
 void FLiveLinkHub::RegisterLiveLinkHubSettings()
 {
 	if (ISettingsModule* SettingsModule = FModuleManager::GetModulePtr<ISettingsModule>("Settings"))
diff --git a/Engine/Plugins/Animation/LiveLinkHub/Source/LiveLinkHub/Private/LiveLinkHub.h b/Engine/Plugins/Animation/LiveLinkHub/Source/LiveLinkHub/Private/LiveLinkHub.h
--- a/Engine/Plugins/Animation/LiveLinkHub/Source/LiveLinkHub/Private/LiveLinkHub.h
+++ b/Engine/Plugins/Animation/LiveLinkHub/Source/LiveLinkHub/Private/LiveLinkHub.h
@@ -80,8 +80,6 @@ private:
 	/** Bind all available live link hub commands. */
 	void BindCommands();
 
-	/** Clear all client settings. */
-	FName GetSubjectNameOverride(const FLiveLinkSubjectKey& InSubjectKey) const;
 	/** Register settings relevant to the livelink hub. */
 	void RegisterLiveLinkHubSettings();
 	/** Unregister settings relevant to the livelink hub. */
diff --git a/Engine/Plugins/Animation/LiveLinkHub/Source/LiveLinkHub/Private/LiveLinkHubClient.cpp b/Engine/Plugins/Animation/LiveLinkHub/Source/LiveLinkHub/Private/LiveLinkHubClient.cpp
--- a/Engine/Plugins/Animation/LiveLinkHub/Source/LiveLinkHub/Private/LiveLinkHubClient.cpp
+++ b/Engine/Plugins/Animation/LiveLinkHub/Source/LiveLinkHub/Private/LiveLinkHubClient.cpp
@@ -41,7 +41,7 @@ void FLiveLinkHubClient::CacheSubjectSettings(const FLiveLinkSubjectKey& Subject
 		ULiveLinkSubjectSettings* SubjectSettings = Settings;
 
 		SubjectItem->GetLiveSubject()->CacheSettings(SourceSettings, SubjectSettings);
-
+		SubjectItem->GetLiveSubject()->SetStaticDataAsRebroadcasted(false);
 		BroadcastStaticDataUpdate(SubjectItem->GetLiveSubject(), SubjectItem->GetSubject()->GetRole(), SubjectItem->GetLiveSubject()->GetStaticData());
 	}
 }
diff --git a/Engine/Plugins/Animation/LiveLinkHub/Source/LiveLinkHub/Private/LiveLinkHubSubjectSettings.cpp b/Engine/Plugins/Animation/LiveLinkHub/Source/LiveLinkHub/Private/LiveLinkHubSubjectSettings.cpp
--- a/Engine/Plugins/Animation/LiveLinkHub/Source/LiveLinkHub/Private/LiveLinkHubSubjectSettings.cpp
+++ b/Engine/Plugins/Animation/LiveLinkHub/Source/LiveLinkHub/Private/LiveLinkHubSubjectSettings.cpp
@@ -11,7 +11,7 @@ void ULiveLinkHubSubjectSettings::Initialize(FLiveLinkSubjectKey InSubjectKey)
 {
 	ILiveLinkClient* LiveLinkClient = &IModularFeatures::Get().GetModularFeature<ILiveLinkClient>(ILiveLinkClient::ModularFeatureName);
 	SubjectName = InSubjectKey.SubjectName.ToString();
-	SubjectKey = InSubjectKey;
+	Key = InSubjectKey;
 
 	OutboundName = SubjectName;
 
@@ -70,7 +70,7 @@ void ULiveLinkHubSubjectSettings::PostEditChangeProperty(FPropertyChangedEvent&
 		|| PropertyChangedEvent.GetPropertyName() == GET_MEMBER_NAME_CHECKED(ULiveLinkHubSubjectSettings, PreProcessors))
 	{
 		FLiveLinkHubClient* LiveLinkClient = static_cast<FLiveLinkHubClient*>(&IModularFeatures::Get().GetModularFeature<ILiveLinkClient>(ILiveLinkClient::ModularFeatureName));
-		LiveLinkClient->CacheSubjectSettings(SubjectKey, this);
+		LiveLinkClient->CacheSubjectSettings(Key, this);
 	}
 	else if (PropertyChangedEvent.GetPropertyName() == GET_MEMBER_NAME_CHECKED(ULiveLinkHubSubjectSettings, TranslatorsProxy))
 	{
@@ -93,6 +93,6 @@ void ULiveLinkHubSubjectSettings::PostEditChangeProperty(FPropertyChangedEvent&
 		}
 
 		FLiveLinkHubClient* LiveLinkClient = static_cast<FLiveLinkHubClient*>(&IModularFeatures::Get().GetModularFeature<ILiveLinkClient>(ILiveLinkClient::ModularFeatureName));
-		LiveLinkClient->CacheSubjectSettings(SubjectKey, this);
+		LiveLinkClient->CacheSubjectSettings(Key, this);
 	}
 }
diff --git a/Engine/Plugins/Animation/LiveLinkHub/Source/LiveLinkHub/Private/LiveLinkHubSubjectSettings.h b/Engine/Plugins/Animation/LiveLinkHub/Source/LiveLinkHub/Private/LiveLinkHubSubjectSettings.h
--- a/Engine/Plugins/Animation/LiveLinkHub/Source/LiveLinkHub/Private/LiveLinkHubSubjectSettings.h
+++ b/Engine/Plugins/Animation/LiveLinkHub/Source/LiveLinkHub/Private/LiveLinkHubSubjectSettings.h
@@ -16,9 +16,15 @@ class ULiveLinkHubSubjectSettings : public ULiveLinkSubjectSettings
 	GENERATED_BODY()
 public:
 
-	/** Initialize the settings. */
+	//~ Begin ULiveLinkSubjectSettings interface
 	virtual void Initialize(FLiveLinkSubjectKey InSubjectKey) override;
 
+	virtual FName GetRebroadcastName() const override
+	{
+		return *OutboundName;
+	}
+	//~ End ULiveLinkSubjectSettings interface
+
 	/** Notify connected clients that this subject's name has changed. */
 	void NotifyRename();
 
@@ -48,9 +54,6 @@ private:
 	/* Previous outbound name to be used for noticing clients to remove this entry from their subject list. */
 	FName PreviousOutboundName;
 
-	/** SubjectKey of this subject. */
-	FLiveLinkSubjectKey SubjectKey;
-
 	struct FNameChangeInfo
 	{
 		UClass* SubjectRole = nullptr;
diff --git a/Engine/Source/Runtime/LiveLinkInterface/Public/LiveLinkSubjectSettings.h b/Engine/Source/Runtime/LiveLinkInterface/Public/LiveLinkSubjectSettings.h
--- a/Engine/Source/Runtime/LiveLinkInterface/Public/LiveLinkSubjectSettings.h
+++ b/Engine/Source/Runtime/LiveLinkInterface/Public/LiveLinkSubjectSettings.h
@@ -55,7 +55,16 @@ public:
 	LIVELINKINTERFACE_API ULiveLinkSubjectSettings();
 
 	/** Initialize the settings. */
-	virtual void Initialize(FLiveLinkSubjectKey InSubjectKey) {}
+	virtual void Initialize(FLiveLinkSubjectKey InSubjectKey)
+	{
+		Key = MoveTemp(InSubjectKey);
+	}
+
+	/** Get the name that should be used when the subject is rebroadcasted. */
+	virtual FName GetRebroadcastName() const
+	{
+		return Key.SubjectName.Name;
+	}
 
 	/** List of available preprocessor the subject will use. */
 	UPROPERTY(EditAnywhere, Instanced, Category = "LiveLink", meta = (DisplayName = "Pre Processors"))
@@ -95,4 +104,9 @@ public:
 	LIVELINKINTERFACE_API virtual void PostEditChangeChainProperty(struct FPropertyChangedChainEvent& PropertyChangedEvent) override;
 #endif
 	//~ End UObject interface
+
+protected:
+	/** Key of the subject that owns this setting. */
+	UPROPERTY()
+	FLiveLinkSubjectKey Key;
 };
