LiveLinkHub - Add  resampling / timecode restamping support

diff --git a/Engine/Plugins/Animation/LiveLink/Source/LiveLink/Public/LiveLinkVirtualSubjectBoneAttachment.h b/Engine/Plugins/Animation/LiveLink/Source/LiveLink/Public/LiveLinkVirtualSubjectBoneAttachment.h
--- a/Engine/Plugins/Animation/LiveLink/Source/LiveLink/Public/LiveLinkVirtualSubjectBoneAttachment.h
+++ b/Engine/Plugins/Animation/LiveLink/Source/LiveLink/Public/LiveLinkVirtualSubjectBoneAttachment.h
@@ -44,14 +44,6 @@ public:
 	UPROPERTY(EditAnywhere, Category = "Live Link", meta=(LinearDeltaSensitivity = "1", Delta = "1.0"))
 	FRotator RotationOffset = FRotator::ZeroRotator;
 
-	/** Whether to ignore the parent bone location when attaching the child bone. */
-	UPROPERTY(EditAnywhere, Category = "Live Link")
-    bool bIgnoreParentLocation = false;
-
-	/** Whether to ignore the parent bone rotation when attaching the child bone. */
-    UPROPERTY(EditAnywhere, Category = "Live Link")
-    bool bIgnoreParentRotation = false;
-
 	/** If not empty, hold the last error text set by the IsValid method. Used to bubble up errors to the UI. */
 	mutable FText LastError;
 };
diff --git a/Engine/Plugins/Animation/LiveLinkHub/Source/LiveLinkHub/Private/LiveLinkHubClient.cpp b/Engine/Plugins/Animation/LiveLinkHub/Source/LiveLinkHub/Private/LiveLinkHubClient.cpp
--- a/Engine/Plugins/Animation/LiveLinkHub/Source/LiveLinkHub/Private/LiveLinkHubClient.cpp
+++ b/Engine/Plugins/Animation/LiveLinkHub/Source/LiveLinkHub/Private/LiveLinkHubClient.cpp
@@ -4,7 +4,9 @@
 
 #include "Clients/LiveLinkHubProvider.h"
 #include "LiveLinkHub.h"
+#include "LiveLinkHubLog.h"
 #include "LiveLinkHubModule.h"
+#include "LiveLinkLog.h"
 #include "LiveLinkRoleTrait.h"
 #include "LiveLinkSourceCollection.h"
 #include "LiveLinkSubject.h"
@@ -128,7 +130,34 @@ void FLiveLinkHubClient::OnStaticDataAdded(FLiveLinkSubjectKey SubjectKey,  TSub
 
 void FLiveLinkHubClient::OnFrameDataAdded(FLiveLinkSubjectKey InSubjectKey, TSubclassOf<ULiveLinkRole> SubjectRole, const FLiveLinkFrameDataStruct& InFrameData)
 {
-	OnFrameDataReceivedDelegate_AnyThread.Broadcast(InSubjectKey, InFrameData);
+	Collection->ForEachSubject([this, InSubjectKey, &InFrameData](const FLiveLinkCollectionSourceItem& SourceItem, const FLiveLinkCollectionSubjectItem& SubjectItem)
+	{
+		if (SourceItem.Setting->ParentSubject.Name == InSubjectKey.SubjectName)
+		{
+			// todo: Time offset evaluation
+			FLiveLinkSubjectFrameData ChildData;
+			if (SubjectItem.GetLiveSubject()->EvaluateFrameAtWorldTime(InFrameData.GetBaseData()->WorldTime.GetSourceTime(), SubjectItem.GetLinkSettings()->Role, ChildData))
+			{
+				const FTimecode FrameTC = FTimecode::FromFrameNumber(InFrameData.GetBaseData()->MetaData.SceneTime.Time.GetFrame(), InFrameData.GetBaseData()->MetaData.SceneTime.Rate);
+				UE_LOG(LogLiveLinkHub, Verbose, TEXT("LiveLinkHub Parent (%s) - Child '%s' adding frame with Timecode:[%s.%0.3f] - SourceTime: %0.4f, Offset: %0.6f, CorrectedTime: %0.4f"), *InSubjectKey.SubjectName.ToString(), *SubjectItem.Key.SubjectName.ToString(), *FrameTC.ToString(), InFrameData.GetBaseData()->MetaData.SceneTime.Time.GetSubFrame(), InFrameData.GetBaseData()->WorldTime.GetSourceTime(), InFrameData.GetBaseData()->WorldTime.GetOffset(), InFrameData.GetBaseData()->WorldTime.GetOffsettedTime());
+
+				ChildData.FrameData.GetBaseData()->MetaData.SceneTime = InFrameData.GetBaseData()->MetaData.SceneTime;
+				ChildData.FrameData.GetBaseData()->MetaData.SceneTime.Rate = InFrameData.GetBaseData()->MetaData.SceneTime.Rate;
+				OnFrameDataReceivedDelegate_AnyThread.Broadcast(SubjectItem.Key, ChildData.FrameData);
+			}
+			else
+			{
+				FLiveLinkLog::Warning(TEXT("Child subjects %s could not be evaluated for data resampling."), *InSubjectKey.SubjectName.Name.ToString());
+			}
+		}
+	});
+
+	const bool bIsParented = Collection->FindSource(InSubjectKey.Source)->Setting->ParentSubject != FLiveLinkSubjectName();
+	if (!bIsParented)
+	{
+		// Don't broadcast frames that have a parent, they'll be broadcast in the loop above.
+		OnFrameDataReceivedDelegate_AnyThread.Broadcast(InSubjectKey, InFrameData);
+	}
 }
 
 
diff --git a/Engine/Plugins/Animation/LiveLinkHub/Source/LiveLinkHub/Private/UI/Widgets/SLiveLinkHubMainTabView.cpp b/Engine/Plugins/Animation/LiveLinkHub/Source/LiveLinkHub/Private/UI/Widgets/SLiveLinkHubMainTabView.cpp
--- a/Engine/Plugins/Animation/LiveLinkHub/Source/LiveLinkHub/Private/UI/Widgets/SLiveLinkHubMainTabView.cpp
+++ b/Engine/Plugins/Animation/LiveLinkHub/Source/LiveLinkHub/Private/UI/Widgets/SLiveLinkHubMainTabView.cpp
@@ -51,7 +51,7 @@ void SLiveLinkHubMainTabView::Construct(const FArguments& InArgs)
 		{
 			CreateTabs(InTabManager, InLayout, InArgs);
 		}))
-		.LayoutName("LiveLinkHubSourcesTabView_v1.2")
+		.LayoutName("LiveLinkHubSourcesTabView_v1.3")
 	);
 }
 
@@ -69,12 +69,8 @@ void SLiveLinkHubMainTabView::CreateTabs(const TSharedRef<FTabManager>& InTabMan
 
 	InTabManager->RegisterTabSpawner(SourcesTabId, FOnSpawnTab::CreateSP(this, &SLiveLinkHubMainTabView::SpawnSourcesTab))
 		.SetDisplayName(SourcesTabName);
-
-	/* Disabled until client details are used in the hub.
 	InTabManager->RegisterTabSpawner(SourceDetailsTabId, FOnSpawnTab::CreateSP(this, &SLiveLinkHubMainTabView::SpawnSourcesDetailsTab))
 		.SetDisplayName(SourceDetailsTabName);
-	*/
-
 	InTabManager->RegisterTabSpawner(SubjectsTabId, FOnSpawnTab::CreateSP(this, &SLiveLinkHubMainTabView::SpawnSubjectsTab))
 		.SetIcon(FSlateIcon(LiveLinkStyleName, TEXT("LiveLinkHub.Subjects.Icon")))
 		.SetDisplayName(SubjectsTabName);
@@ -107,6 +103,12 @@ void SLiveLinkHubMainTabView::CreateTabs(const TSharedRef<FTabManager>& InTabMan
 							->SetHideTabWell(true)
 							->AddTab(SourcesTabId, ETabState::OpenedTab)
 						)
+						->Split
+						(
+							FTabManager::NewStack()
+							->SetSizeCoefficient(0.5f)
+							->AddTab(SourceDetailsTabId, ETabState::OpenedTab)
+						)
 					)
 					->Split
 					(
diff --git a/Engine/Source/Runtime/LiveLinkInterface/Public/LiveLinkSourceSettings.h b/Engine/Source/Runtime/LiveLinkInterface/Public/LiveLinkSourceSettings.h
--- a/Engine/Source/Runtime/LiveLinkInterface/Public/LiveLinkSourceSettings.h
+++ b/Engine/Source/Runtime/LiveLinkInterface/Public/LiveLinkSourceSettings.h
@@ -184,6 +184,15 @@ public:
 	UPROPERTY(VisibleAnywhere, AdvancedDisplay, Category = "Settings")
 	TSubclassOf<ULiveLinkSourceFactory> Factory;
 
+	/** 
+	 * Parent subject used for resampling.
+	 * If this is set, this source's subjects will only be updated when the parent subject receives data.
+	 * Additionally this source's subjects' timecode will match the parent's subject received timecode.
+	 * This can be useful for synchronizing a higher frequency source to a lower frequency one.
+	 */
+	UPROPERTY(EditAnywhere, Category = "Settings")
+	FLiveLinkSubjectName ParentSubject;
+
 #if WITH_EDITORONLY_DATA
 	UPROPERTY()
 	TArray<FLiveLinkSourceDebugInfo> SourceDebugInfos_DEPRECATED;
