LiveLinkHub - Fix resetting outbound name property resetting to None instead of the expected value

diff --git a/Engine/Plugins/Animation/LiveLinkHub/Source/LiveLinkHub/Private/Subjects/LiveLinkHubSubjectSettingsDetailsCustomization.cpp b/Engine/Plugins/Animation/LiveLinkHub/Source/LiveLinkHub/Private/Subjects/LiveLinkHubSubjectSettingsDetailsCustomization.cpp
--- a/Engine/Plugins/Animation/LiveLinkHub/Source/LiveLinkHub/Private/Subjects/LiveLinkHubSubjectSettingsDetailsCustomization.cpp
+++ b/Engine/Plugins/Animation/LiveLinkHub/Source/LiveLinkHub/Private/Subjects/LiveLinkHubSubjectSettingsDetailsCustomization.cpp
@@ -3,6 +3,8 @@
 #include "LiveLinkHubSubjectSettingsDetailsCustomization.h"
 
 #include "DetailLayoutBuilder.h"
+#include "DetailWidgetRow.h"
+#include "IDetailPropertyRow.h"
 #include "LiveLinkHubSubjectSettings.h"
 
 #define LOCTEXT_NAMESPACE "LiveLinkHubSubjectSettingsDetailCustomization"
@@ -17,6 +19,49 @@ void FLiveLinkHubSubjectSettingsDetailsCustomization::CustomizeDetails(IDetailLa
 		return;
 	}
 
+	TSharedRef<IPropertyHandle> OutboundProperty = InDetailBuilder.GetProperty(GET_MEMBER_NAME_CHECKED(ULiveLinkHubSubjectSettings, OutboundName));
+	if (IDetailPropertyRow* PropertyRow = InDetailBuilder.EditDefaultProperty(OutboundProperty))
+	{
+		FResetToDefaultOverride ResetOverride;
+
+		PropertyRow->CustomWidget()
+			.OverrideResetToDefault(FResetToDefaultOverride::Create(
+				FIsResetToDefaultVisible::CreateLambda([](TSharedPtr<IPropertyHandle> PropertyHandle)
+					{
+						TSharedPtr<IPropertyHandle> SubjectNameProperty = PropertyHandle->GetParentHandle()->GetChildHandle(GET_MEMBER_NAME_CHECKED(ULiveLinkHubSubjectSettings, SubjectName));
+						
+						if (SubjectNameProperty && SubjectNameProperty->IsValidHandle())
+						{
+							FText OutboundName;
+							FText SubjectName;
+							PropertyHandle->GetValueAsDisplayText(OutboundName);
+							SubjectNameProperty->GetValueAsDisplayText(SubjectName);
+							return !OutboundName.EqualTo(SubjectName);
+						}
+						return false;
+					}),
+				FResetToDefaultHandler::CreateLambda([](TSharedPtr<IPropertyHandle> PropertyHandle)
+					{
+						TSharedPtr<IPropertyHandle> SubjectNameProperty = PropertyHandle->GetParentHandle()->GetChildHandle(GET_MEMBER_NAME_CHECKED(ULiveLinkHubSubjectSettings, SubjectName));
+
+						if (SubjectNameProperty && SubjectNameProperty->IsValidHandle())
+						{
+							FString SubjectName;
+							SubjectNameProperty->GetValueAsDisplayString(SubjectName);
+							PropertyHandle->SetValue(*SubjectName);
+						}
+					})
+			))
+			.NameContent()
+			[
+				OutboundProperty->CreatePropertyNameWidget()
+			]
+			.ValueContent()
+			[
+				OutboundProperty->CreatePropertyValueWidget()
+			];
+	}
+
 	InDetailBuilder.HideProperty(GET_MEMBER_NAME_CHECKED(ULiveLinkHubSubjectSettings, InterpolationProcessor), ULiveLinkSubjectSettings::StaticClass());
 	InDetailBuilder.HideProperty(GET_MEMBER_NAME_CHECKED(ULiveLinkHubSubjectSettings, FrameRate), ULiveLinkSubjectSettings::StaticClass());
 	InDetailBuilder.HideProperty(GET_MEMBER_NAME_CHECKED(ULiveLinkHubSubjectSettings, bRebroadcastSubject), ULiveLinkSubjectSettings::StaticClass());
