LiveLinkHub - Animation data is now stored as bulk data directly on the recording UAsset.
- Async loading the animation in chunks is supported.
- Saving a recording still requires all animation data loaded into memory, but writing the package to disk is now an async operation.
- Playing back a recording immediately after recording while the asset is still saving is supported.
- The source list and details is now read only when playing back a recording or while recording.
- Handle renaming an asset updating the animation data correctly.
- Optimize how the animation buffer size is reported, making the UI buffering smoother and reducing blocking tasks on larger animations with a large buffer size.
- The frame rate of the recording is now saved with the recording and is used when loading rather than the per frame sampe rate. However, framerate calculations still need some work, as the length of recording and correct frame index-to-timestamp can still be incorrect.
- Fix crash when starting a recording quickly after exiting playback and restoring previous sources.

diff --git a/Engine/Plugins/Animation/LiveLink/Source/LiveLink/Private/LiveLinkPreset.cpp b/Engine/Plugins/Animation/LiveLink/Source/LiveLink/Private/LiveLinkPreset.cpp
--- a/Engine/Plugins/Animation/LiveLink/Source/LiveLink/Private/LiveLinkPreset.cpp
+++ b/Engine/Plugins/Animation/LiveLink/Source/LiveLink/Private/LiveLinkPreset.cpp
@@ -260,6 +260,12 @@ void ULiveLinkPreset::ApplyToClientLatent(TFunction<void(bool)> CompletionCallba
 	UE_LOG(LogTemp, Verbose, TEXT("Applying preset %s (%d)"), *GetFullName(), ApplyCount++);
 }
 
+void ULiveLinkPreset::CancelLatentAction()
+{
+	ClearApplyToClientTimer();
+	ApplyToClientPollingOperation.Reset();
+}
+
 bool ULiveLinkPreset::AddToClient(const bool bRecreatePresets) const
 {
 	bool bResult = false;
diff --git a/Engine/Plugins/Animation/LiveLink/Source/LiveLink/Public/LiveLinkPreset.h b/Engine/Plugins/Animation/LiveLink/Source/LiveLink/Public/LiveLinkPreset.h
--- a/Engine/Plugins/Animation/LiveLink/Source/LiveLink/Public/LiveLinkPreset.h
+++ b/Engine/Plugins/Animation/LiveLink/Source/LiveLink/Public/LiveLinkPreset.h
@@ -45,6 +45,9 @@ public:
 	UFUNCTION(BlueprintCallable, BlueprintPure = false, Category="LiveLink", meta = (Latent, LatentInfo = "LatentInfo", HidePin = "WorldContextObject", DefaultToSelf = "WorldContextObject"))
 	void ApplyToClientLatent(UObject* WorldContextObject, FLatentActionInfo LatentInfo);
 	void ApplyToClientLatent(TFunction<void(bool)> CompletionCallback = nullptr);
+	
+	/** Cancels the current latent action and prevents the callback from firing. Only valid for C++ variant. */
+	void CancelLatentAction();
 
 	/**
 	 * Add the sources and subjects from this preset, but leave any existing sources and subjects connected.
diff --git a/Engine/Plugins/Animation/LiveLinkHub/Source/LiveLinkHub/Private/Data/LiveLinkHubBulkData.cpp b/Engine/Plugins/Animation/LiveLinkHub/Source/LiveLinkHub/Private/Data/LiveLinkHubBulkData.cpp
new file mode 100644
--- /dev/null
+++ b/Engine/Plugins/Animation/LiveLinkHub/Source/LiveLinkHub/Private/Data/LiveLinkHubBulkData.cpp
@@ -0,0 +1,89 @@
+// Copyright Epic Games, Inc. All Rights Reserved.
+
+#include "LiveLinkHubBulkData.h"
+
+#include "Serialization/BufferArchive.h"
+#include "Serialization/MemoryReader.h"
+
+FLiveLinkHubBulkData::FScopedBulkDataMemoryReader::FScopedBulkDataMemoryReader(const int64 InOffset, const int32 InBytesToRead, FLiveLinkHubBulkData* InBulkData)
+{
+	Memory.SetNumUninitialized(InBytesToRead);
+
+	LocalBulkDataOffset = InBulkData->ReadBulkDataImpl(InOffset, InBytesToRead, Memory.GetData());
+	MemoryReader = MakeUnique<FMemoryReader>(Memory, true);
+}
+
+FLiveLinkHubBulkData::~FLiveLinkHubBulkData()
+{
+	CloseFileReader();
+}
+
+void FLiveLinkHubBulkData::CloseFileReader()
+{
+	RecordingFileReader.Reset();
+}
+
+void FLiveLinkHubBulkData::ReadBulkData(const int64 InBytesToRead, uint8* InMemory)
+{
+	BulkDataOffset = ReadBulkDataImpl(BulkDataOffset, InBytesToRead, InMemory);
+}
+
+FLiveLinkHubBulkData::FScopedBulkDataMemoryReader FLiveLinkHubBulkData::CreateBulkDataMemoryReader(const int64 InBytesToRead)
+{
+	FScopedBulkDataMemoryReader Reader(BulkDataOffset, InBytesToRead, this);
+	BulkDataOffset = Reader.GetBulkDataOffset();
+
+	return Reader;
+}
+
+void FLiveLinkHubBulkData::ResetBulkDataOffset()
+{
+	BulkDataOffset = BulkData.GetBulkDataOffsetInFile();
+}
+
+void FLiveLinkHubBulkData::SetBulkDataOffset(const int64 InNewOffset)
+{
+	BulkDataOffset = InNewOffset;
+}
+
+void FLiveLinkHubBulkData::Serialize(FArchive& Ar, UObject* Owner)
+{
+	BulkData.SetBulkDataFlags(BULKDATA_Force_NOT_InlinePayload);
+	BulkData.Serialize(Ar, Owner);
+}
+
+void FLiveLinkHubBulkData::WriteBulkData(FBufferArchive& BufferArchive)
+{
+	BulkData.Lock(LOCK_READ_WRITE);
+	BulkData.Realloc(BufferArchive.Num());
+	unsigned char* BulkDataPtr = BulkData.Realloc(BufferArchive.Num());
+	FMemory::Memcpy(BulkDataPtr, BufferArchive.GetData(), BufferArchive.Num());
+	BulkData.Unlock();
+}
+
+int64 FLiveLinkHubBulkData::ReadBulkDataImpl(int64 InOffset, int64 InBytesToRead, uint8* InMemory)
+{
+	if (!RecordingFileReader.IsValid())
+	{
+		check(BulkData.DoesExist());
+		check(BulkData.CanLoadFromDisk());
+		check(!BulkData.IsInlined());
+		check(!BulkData.IsInSeparateFile());
+		check(!BulkData.IsBulkDataLoaded());
+
+		RecordingFileReader = TUniquePtr<IAsyncReadFileHandle>(BulkData.OpenAsyncReadHandle());
+	}
+
+	check(RecordingFileReader);
+
+	const TUniquePtr<IAsyncReadRequest> ReadRequest(RecordingFileReader->ReadRequest(
+		InOffset,
+		InBytesToRead,
+		AIOP_High,
+		nullptr,
+		InMemory
+		));
+	ReadRequest->WaitCompletion();
+
+	return InOffset + InBytesToRead;
+}
diff --git a/Engine/Plugins/Animation/LiveLinkHub/Source/LiveLinkHub/Private/Data/LiveLinkHubBulkData.h b/Engine/Plugins/Animation/LiveLinkHub/Source/LiveLinkHub/Private/Data/LiveLinkHubBulkData.h
new file mode 100644
--- /dev/null
+++ b/Engine/Plugins/Animation/LiveLinkHub/Source/LiveLinkHub/Private/Data/LiveLinkHubBulkData.h
@@ -0,0 +1,86 @@
+// Copyright Epic Games, Inc. All Rights Reserved.
+
+#pragma once
+
+#include "Serialization/BulkData.h"
+
+class FBufferArchive;
+class FMemoryReader;
+class IAsyncReadFileHandle;
+
+/** Bulk data for live link hub, automatically tracking the current offset of the bulk data. */
+class FLiveLinkHubBulkData
+{
+public:
+	/**
+	 * Read a section of bulk data to a memory reader - valid only within this scope - allowing bulk data to be used like a typical FArchive.
+	 * This differs from FBulkDataReader in that it supports only loading a section of bulk data into memory at a time.
+	 */
+	class FScopedBulkDataMemoryReader
+	{
+	public:
+		FScopedBulkDataMemoryReader(const int64 InOffset, const int32 InBytesToRead, FLiveLinkHubBulkData* InBulkData);
+
+		/** Retrieve the memory reader. */
+		FMemoryReader& GetMemoryReader() const { return *MemoryReader; }
+		
+		/** Retrieve the bulk data offset after having read the bulk data into memory. */
+		int64 GetBulkDataOffset() const { return LocalBulkDataOffset; }
+		
+	private:
+		/** The memory reader for reading out memory storage. */
+		TUniquePtr<FMemoryReader> MemoryReader;
+		/** The memory storage for loading the bulk data. */
+		TArray<uint8> Memory;
+		/** The bulk data offset after reading the data into memory. */
+		int64 LocalBulkDataOffset = 0;
+	};
+
+	~FLiveLinkHubBulkData();
+
+	/** Close the file reader if it is open. */
+	void CloseFileReader();
+	
+	/** Read data using sizeof(T). */
+	template<typename T>
+	void ReadBulkDataPrimitive(T& InMemory)
+	{
+		ReadBulkData(sizeof(T), reinterpret_cast<uint8*>(&InMemory));
+	}
+
+	/** Read bulk data at the current offset of the specified byte size and increment the file offset. */
+	void ReadBulkData(const int64 InBytesToRead, uint8* InMemory);
+
+	/** Create a scoped memory reader consisting of the bulk data bytes read. Increments file offset. */
+	FScopedBulkDataMemoryReader CreateBulkDataMemoryReader(const int64 InBytesToRead);
+	
+	/** Reset to the initial offset in the file. */
+	void ResetBulkDataOffset();
+
+	/** Manually set the bulk data offset. */
+	void SetBulkDataOffset(const int64 InNewOffset);
+
+	/** Retrieve the current bulk data offset. */
+	int64 GetBulkDataOffset() const { return BulkDataOffset; }
+
+	/** Call the serialize method of the bulk data, should be called when serializing the owning asset. */
+	void Serialize(FArchive& Ar, UObject* Owner);
+
+	/** Write to the bulk data file. */
+	void WriteBulkData(FBufferArchive& BufferArchive);
+	
+private:
+	/**
+	 * Read bulk data at an offset of the specified byte size.
+	 * @return The new offset after reading in the bytes.
+	 */
+	int64 ReadBulkDataImpl(int64 InOffset, int64 InBytesToRead, uint8* InMemory);
+	
+private:
+	/** The file reader open to the bulk data. */
+	TUniquePtr<IAsyncReadFileHandle> RecordingFileReader;
+	/** The bulk data storage. */
+	FByteBulkData BulkData;
+	/** The current bulk data offset in the owning file. */
+	int64 BulkDataOffset = 0;
+};
diff --git a/Engine/Plugins/Animation/LiveLinkHub/Source/LiveLinkHub/Private/Recording/Implementations/LiveLinkUAssetRecorder.cpp b/Engine/Plugins/Animation/LiveLinkHub/Source/LiveLinkHub/Private/Recording/Implementations/LiveLinkUAssetRecorder.cpp
--- a/Engine/Plugins/Animation/LiveLinkHub/Source/LiveLinkHub/Private/Recording/Implementations/LiveLinkUAssetRecorder.cpp
+++ b/Engine/Plugins/Animation/LiveLinkHub/Source/LiveLinkHub/Private/Recording/Implementations/LiveLinkUAssetRecorder.cpp
@@ -21,12 +21,14 @@
 #include "Misc/FileHelper.h"
 #include "Misc/PackageName.h"
 #include "Recording/Implementations/LiveLinkUAssetRecording.h"
+#include "Recording/LiveLinkHubPlaybackController.h"
 #include "Recording/LiveLinkRecording.h"
 #include "Settings/LiveLinkHubSettings.h"
 #include "StructUtils/InstancedStruct.h"
 #include "UI/Window/LiveLinkHubWindowController.h"
 #include "UObject/Object.h"
 #include "UObject/Package.h"
+#include "UObject/SavePackage.h"
 #include "UObject/StrongObjectPtr.h"
 #include "Widgets/DeclarativeSyntaxSupport.h"
 #include "Widgets/Layout/SBox.h"
@@ -45,15 +47,16 @@ namespace UAssetRecorderUtils
 		
 		if (const FLiveLinkStaticDataStruct* StaticData = LiveLinkClient->GetSubjectStaticData_AnyThread(SubjectKey))
 		{
-			check(StaticData->IsValid());
-		
-			TSharedPtr<FInstancedStruct> StaticDataInstancedStruct = MakeShared<FInstancedStruct>();
-			StaticDataInstancedStruct->InitializeAs(StaticData->GetStruct(), (uint8*)StaticData->GetBaseData());
-
-			StaticDataContainer = FLiveLinkRecordingStaticDataContainer();
-			StaticDataContainer->Role = LiveLinkRole;
-			StaticDataContainer->RecordedData.Insert(MoveTemp(StaticDataInstancedStruct), 0);
-			StaticDataContainer->Timestamps.Add(0.0);
+			if (StaticData->IsValid())
+			{
+				TSharedPtr<FInstancedStruct> StaticDataInstancedStruct = MakeShared<FInstancedStruct>();
+				StaticDataInstancedStruct->InitializeAs(StaticData->GetStruct(), (uint8*)StaticData->GetBaseData());
+
+				StaticDataContainer = FLiveLinkRecordingStaticDataContainer();
+				StaticDataContainer->Role = LiveLinkRole;
+				StaticDataContainer->RecordedData.Insert(MoveTemp(StaticDataInstancedStruct), 0);
+				StaticDataContainer->Timestamps.Add(0.0);
+			}
 		}
 
 		return StaticDataContainer;
@@ -88,6 +91,11 @@ bool FLiveLinkUAssetRecorder::IsRecording() const
 	return bIsRecording;
 }
 
+bool FLiveLinkUAssetRecorder::IsSavingRecording(ULiveLinkRecording* InRecording) const
+{
+	return AsyncSaveTasks.Contains(InRecording);
+}
+
 void FLiveLinkUAssetRecorder::RecordBaseData(FLiveLinkRecordingBaseDataContainer& StaticDataContainer, TSharedPtr<FInstancedStruct>&& DataToRecord)
 {
 	const double TimeNowInSeconds = FPlatformTime::Seconds();
@@ -220,20 +228,19 @@ void FLiveLinkUAssetRecorder::SaveRecording()
 	const FString NewAssetName = FPackageName::GetLongPackageAssetName(PackageName);
 	UPackage* NewPackage = CreatePackage(*PackageName);
 
-	ULiveLinkUAssetRecording* NewRecording = NewObject<ULiveLinkUAssetRecording>(NewPackage, *NewAssetName, RF_Public | RF_Standalone);
-	if (NewRecording)
+	if (TObjectPtr<ULiveLinkUAssetRecording> NewRecording = NewObject<ULiveLinkUAssetRecording>(NewPackage, *NewAssetName, RF_Public | RF_Standalone))
 	{
-		NewRecording->LengthInSeconds = TimeRecordingEnded - TimeRecordingStarted;
-		NewRecording->RecordingPreset->BuildFromClient();
-		NewRecording->RecordingData = MoveTemp(*CurrentRecording);
+		const double RecordingLength = TimeRecordingEnded - TimeRecordingStarted;
+		NewRecording->InitializeNewRecordingData(MoveTemp(*CurrentRecording), RecordingLength);
 
-		// todo: async stream save and utilize bulk data within UAsset instead of a separate file.
-		NewRecording->SaveRecording();
-		
 		NewRecording->MarkPackageDirty();
 
-		//FAssetRegistryModule::AssetCreated(NewPreset);  Disabled for now since unsure if needed.
-		UEditorLoadingAndSavingUtils::SavePackages({ NewPackage }, false);
+		FAssetRegistryModule::AssetCreated(NewRecording);
+
+		// Create a task to run on a separate thread for saving all frame data and writing the final UAsset to disk.
+		// We use a container rather than just one task on the chance a save operation is still running when another recording is being saved.
+		FAsyncTask<FLiveLinkSaveRecordingAsyncTask>& AsyncTask = *AsyncSaveTasks.Add(NewRecording, MakeUnique<FAsyncTask<FLiveLinkSaveRecordingAsyncTask>>(NewRecording, this));
+		AsyncTask.StartBackgroundTask();
 	}
 
 	return;
@@ -254,4 +261,60 @@ void FLiveLinkUAssetRecorder::RecordInitialStaticData()
 	}
 }
 
+void FLiveLinkUAssetRecorder::OnRecordingSaved_GameThread(TWeakObjectPtr<ULiveLinkUAssetRecording> InRecording)
+{
+	if (InRecording.IsValid())
+	{
+		const FLiveLinkHubModule& LiveLinkHubModule = FModuleManager::Get().GetModuleChecked<FLiveLinkHubModule>("LiveLinkHub");
+		const TStrongObjectPtr<ULiveLinkRecording> PlaybackRecording = LiveLinkHubModule.GetPlaybackController()->GetRecording();
+
+		UPackage* PackageToUnload = InRecording->GetPackage();
+		const bool bIsPlayingThisRecording = PlaybackRecording.Get() == InRecording.Get();
+
+		// Finish task first to make sure strong reference to the recording is cleared.
+		if (const TUniquePtr<FAsyncTask<FLiveLinkSaveRecordingAsyncTask>>* AsyncTask = AsyncSaveTasks.Find(InRecording))
+		{
+			AsyncTask->Get()->EnsureCompletion();
+			ensure(AsyncSaveTasks.Remove(InRecording) > 0);
+		}
+		
+		if (!bIsPlayingThisRecording)
+		{
+			// Unload as this is not used again until the user loads it, and allows the bulk animation data to obtain a file handle correctly.
+			LiveLinkHubModule.GetPlaybackController()->UnloadRecordingPackage(PackageToUnload);
+		}
+	}
+}
+
+void FLiveLinkUAssetRecorder::FLiveLinkSaveRecordingAsyncTask::DoWork()
+{
+	check(LiveLinkRecording.IsValid());
+
+	// Write to bulk data.
+	LiveLinkRecording->SaveRecordingData();
+		
+	FSavePackageArgs SavePackageArgs;
+	SavePackageArgs.TopLevelFlags = RF_Public | RF_Standalone;
+	SavePackageArgs.Error = GLog;
+	SavePackageArgs.SaveFlags = SAVE_Async;
+
+	const FString PackageFileName = FPackageName::LongPackageNameToFilename(LiveLinkRecording->GetPackage()->GetName(), FPackageName::GetAssetPackageExtension());
+	if (UPackage::SavePackage(LiveLinkRecording->GetPackage(), LiveLinkRecording.Get(), *PackageFileName, MoveTemp(SavePackageArgs)))
+	{
+		UPackage::WaitForAsyncFileWrites();
+
+		// Finish on the game thread for safety.
+		TWeakObjectPtr RecordingWeakPtr(LiveLinkRecording.Get());
+		FSimpleDelegateGraphTask::CreateAndDispatchWhenReady(
+			FSimpleDelegateGraphTask::FDelegate::CreateRaw(Recorder, &FLiveLinkUAssetRecorder::OnRecordingSaved_GameThread, RecordingWeakPtr),
+		TStatId(),
+		nullptr,
+		ENamedThreads::GameThread);
+	}
+	else
+	{
+		UE_LOG(LogLiveLinkHub, Error, TEXT("Package '%s' was not saved"), *PackageFileName);
+	}
+}
+
 #undef LOCTEXT_NAMESPACE
diff --git a/Engine/Plugins/Animation/LiveLinkHub/Source/LiveLinkHub/Private/Recording/Implementations/LiveLinkUAssetRecorder.h b/Engine/Plugins/Animation/LiveLinkHub/Source/LiveLinkHub/Private/Recording/Implementations/LiveLinkUAssetRecorder.h
--- a/Engine/Plugins/Animation/LiveLinkHub/Source/LiveLinkHub/Private/Recording/Implementations/LiveLinkUAssetRecorder.h
+++ b/Engine/Plugins/Animation/LiveLinkHub/Source/LiveLinkHub/Private/Recording/Implementations/LiveLinkUAssetRecorder.h
@@ -2,10 +2,11 @@
 
 #pragma once
 
-#include "Recording/LiveLinkRecorder.h"
-
+#include "Async/AsyncWork.h"
 #include "LiveLinkTypes.h"
+#include "LiveLinkUAssetRecording.h"
 #include "Misc/CoreMiscDefines.h"
+#include "Recording/LiveLinkRecorder.h"
 #include "Templates/PimplPtr.h"
 
 struct FLiveLinkRecordingBaseDataContainer;
@@ -23,6 +24,7 @@ public:
 	virtual void RecordStaticData(const FLiveLinkSubjectKey& SubjectKey, TSubclassOf<ULiveLinkRole> Role, const FLiveLinkStaticDataStruct& StaticData) override;
 	virtual void RecordFrameData(const FLiveLinkSubjectKey& SubjectKey, const FLiveLinkFrameDataStruct& FrameData) override;
 	virtual bool IsRecording() const override;
+	virtual bool IsSavingRecording(ULiveLinkRecording* InRecording) const override;
 	//~ End ILiveLinkRecorder
 
 private:
@@ -36,8 +38,35 @@ private:
 	void RecordBaseData(FLiveLinkRecordingBaseDataContainer& StaticDataContainer, TSharedPtr<FInstancedStruct>&& DataToRecord);
 	/** Record initial data for all livelink subjects. (Useful when static data was sent before the recording started). */
 	void RecordInitialStaticData();
+	/** Called on the game thread after a recording has been saved. */
+	void OnRecordingSaved_GameThread(TWeakObjectPtr<ULiveLinkUAssetRecording> InRecording);
 
 private:
+	class FLiveLinkSaveRecordingAsyncTask : public FNonAbandonableTask
+	{
+	public:
+		FLiveLinkSaveRecordingAsyncTask(ULiveLinkUAssetRecording* InLiveLinkRecording, FLiveLinkUAssetRecorder* InRecorder)
+		{
+			LiveLinkRecording = TStrongObjectPtr(InLiveLinkRecording);
+			Recorder = InRecorder;
+		}
+
+		TStatId GetStatId() const
+		{
+			RETURN_QUICK_DECLARE_CYCLE_STAT(LiveLinkSaveRecordingAsyncTask, STATGROUP_ThreadPoolAsyncTasks);
+		}
+
+		void DoWork();
+
+	private:
+		/** The recording being saved. */
+		TStrongObjectPtr<ULiveLinkUAssetRecording> LiveLinkRecording;
+		/** The recorder owner. */
+		FLiveLinkUAssetRecorder* Recorder = nullptr;
+	};
+
+	/** Current async save tasks. */
+	TMap<TWeakObjectPtr<ULiveLinkRecording>, TUniquePtr<FAsyncTask<FLiveLinkSaveRecordingAsyncTask>>> AsyncSaveTasks;
 	/** Holds metadata and recording data. */
 	TPimplPtr<FLiveLinkUAssetRecordingData> CurrentRecording;
 	/** Whether we're currently recording livelink data. */
diff --git a/Engine/Plugins/Animation/LiveLinkHub/Source/LiveLinkHub/Private/Recording/Implementations/LiveLinkUAssetRecording.cpp b/Engine/Plugins/Animation/LiveLinkHub/Source/LiveLinkHub/Private/Recording/Implementations/LiveLinkUAssetRecording.cpp
--- a/Engine/Plugins/Animation/LiveLinkHub/Source/LiveLinkHub/Private/Recording/Implementations/LiveLinkUAssetRecording.cpp
+++ b/Engine/Plugins/Animation/LiveLinkHub/Source/LiveLinkHub/Private/Recording/Implementations/LiveLinkUAssetRecording.cpp
@@ -2,60 +2,98 @@
 
 #include "LiveLinkUAssetRecording.h"
 
+#include "Engine/Engine.h"
+#include "Features/IModularFeatures.h"
 #include "HAL/IConsoleManager.h"
+#include "LiveLinkHubClient.h"
 #include "LiveLinkHubLog.h"
+#include "LiveLinkHubModule.h"
 #include "LiveLinkUAssetRecordingPlayer.h"
+#include "Misc/App.h"
 #include "Misc/FileHelper.h"
 #include "Misc/Paths.h"
+#include "Modules/ModuleManager.h"
+#include "Recording/LiveLinkHubPlaybackController.h"
+#include "Recording/LiveLinkHubRecordingController.h"
 #include "Serialization/BufferArchive.h"
 #include "Serialization/MemoryReader.h"
 #include "Serialization/ObjectAndNameAsStringProxyArchive.h"
+#include "UObject/Package.h"
 
 ULiveLinkUAssetRecording::~ULiveLinkUAssetRecording()
 {
-	UnloadRecording();
+	UnloadRecordingData();
 }
 
 void ULiveLinkUAssetRecording::Serialize(FArchive& Ar)
 {
 	Super::Serialize(Ar);
+	AnimationData.Serialize(Ar, this);
 }
 
-void ULiveLinkUAssetRecording::SaveRecording()
+void ULiveLinkUAssetRecording::PostDuplicate(EDuplicateMode::Type DuplicateMode)
 {
-	// Saving frame data to a custom file, so it can be streamed in later.
-	const FString FilePath = GetRecordingDataFilePath();
-	FArchive* FileWriter = IFileManager::Get().CreateFileWriter(*FilePath);
+	Super::PostDuplicate(DuplicateMode);
 
+	if (DuplicateMode == EDuplicateMode::Normal)
+	{
+		EjectAndUnload();
+	}
+}
+
+void ULiveLinkUAssetRecording::PostRename(UObject* OldOuter, const FName OldName)
+{
+	Super::PostRename(OldOuter, OldName);
+
+	EjectAndUnload();
+}
+
+void ULiveLinkUAssetRecording::SaveRecordingData()
+{
+	bIsSavingRecordingData = true;
+	
+	FBufferArchive Archive;
+	
 	int32 RecordingVersionToSave = RecordingVersion;
-	*FileWriter << RecordingVersionToSave;
+	Archive << RecordingVersionToSave;
 
 	// How much static data to expect.
 	int32 NumStaticData = RecordingData.StaticData.Num();
-	*FileWriter << NumStaticData;
+	Archive << NumStaticData;
 		
 	for (TTuple<FLiveLinkSubjectKey, FLiveLinkRecordingStaticDataContainer>& StaticData : RecordingData.StaticData)
 	{
-		SaveFrameData(FileWriter, StaticData.Key, StaticData.Value);
+		SaveFrameData(&Archive, StaticData.Key, StaticData.Value);
 	}
 	
 	// How much frame data to expect.
 	int32 NumFrameData = RecordingData.FrameData.Num();
-	*FileWriter << NumFrameData;
+	Archive << NumFrameData;
 		
 	for (TTuple<FLiveLinkSubjectKey, FLiveLinkRecordingBaseDataContainer>& FrameData : RecordingData.FrameData)
 	{
-		SaveFrameData(FileWriter, FrameData.Key, FrameData.Value);
+		SaveFrameData(&Archive, FrameData.Key, FrameData.Value);
 	}
-		
-	FileWriter->FlushCache();
-	FileWriter->Close();
-	delete FileWriter;
+
+	AnimationData.WriteBulkData(Archive);
+	
+	Archive.FlushCache();
+	Archive.Close();
+
+	bIsSavingRecordingData = false;
 }
 
-void ULiveLinkUAssetRecording::LoadRecording(int32 InInitialFrame, int32 InNumFramesToLoad)
+void ULiveLinkUAssetRecording::LoadRecordingData(int32 InInitialFrame, int32 InNumFramesToLoad)
 {
+	if (bIsFullyLoaded)
+	{
+		return;
+	}
+	
 	bCancelStream = false;
+	bPauseStream = false;
+	OnStreamPausedEvent->Reset();
+	OnStreamUnpausedEvent->Reset();
 	
 	int32 StartFrame = InInitialFrame - InNumFramesToLoad;
 	if (StartFrame < 0)
@@ -69,18 +107,7 @@ void ULiveLinkUAssetRecording::LoadRecording(int32 InInitialFrame, int32 InNumFr
 	// Perform initial setup of the file reader.
 	if (!AsyncStreamTask.IsValid())
 	{
-		check(RecordingFileReader == nullptr);
-
 		FrameFileData.Empty();
-		
-		const FString FilePath = GetRecordingDataFilePath();
-		RecordingFileReader = IFileManager::Get().CreateFileReader(*FilePath);
-
-		if (RecordingFileReader == nullptr)
-		{
-			UE_LOG(LogLiveLinkHub, Error, TEXT("Failed to open file %s for reading."), *FilePath);
-			return;
-		}
 	}
 
 	EarliestFrameToStream = StartFrame;
@@ -93,6 +120,16 @@ void ULiveLinkUAssetRecording::LoadRecording(int32 InInitialFrame, int32 InNumFr
 	InitialFrameToStream = InInitialFrame;
 	TotalFramesToStream = InNumFramesToLoad;
 
+	if (!OnPreGarbageCollectHandle.IsValid())
+	{
+		OnPreGarbageCollectHandle = FCoreUObjectDelegates::GetPreGarbageCollectDelegate().AddUObject(this, &ULiveLinkUAssetRecording::OnPreGarbageCollect);
+	}
+
+	if (!OnPostGarbageCollectHandle.IsValid())
+	{
+		OnPostGarbageCollectHandle = FCoreUObjectDelegates::GetPostGarbageCollect().AddUObject(this, &ULiveLinkUAssetRecording::OnPostGarbageCollect);
+	}
+	
 	if (!AsyncStreamTask.IsValid())
 	{
 		AsyncStreamTask = MakeUnique<FAsyncTask<FLiveLinkStreamAsyncTask>>(this);
@@ -100,9 +137,17 @@ void ULiveLinkUAssetRecording::LoadRecording(int32 InInitialFrame, int32 InNumFr
 	}
 }
 
-void ULiveLinkUAssetRecording::UnloadRecording()
+void ULiveLinkUAssetRecording::UnloadRecordingData()
 {
+	if (IsSavingRecordingData() || GetPackage()->HasAnyPackageFlags(PKG_IsSaving))
+	{
+		UE_LOG(LogLiveLinkHub, Warning, TEXT("Attempted to unload %s while the package was still being saved"), *GetName());
+		return;
+	}
+	
 	bCancelStream = true;
+	UnpauseStream();
+	bIsFullyLoaded = false;
 
 	if (AsyncStreamTask.IsValid())
 	{
@@ -113,15 +158,22 @@ void ULiveLinkUAssetRecording::UnloadRecording()
 		AsyncStreamTask.Reset();
 	}
 
-	bPerformedInitialLoad = false;
+	if (OnPreGarbageCollectHandle.IsValid())
+	{
+		FCoreUObjectDelegates::GetPreGarbageCollectDelegate().Remove(OnPreGarbageCollectHandle);
+		OnPreGarbageCollectHandle.Reset();
+	}
 
-	if (RecordingFileReader)
+	if (OnPostGarbageCollectHandle.IsValid())
 	{
-		RecordingFileReader->Close();
-		delete RecordingFileReader;
-		RecordingFileReader = nullptr;
+		FCoreUObjectDelegates::GetPostGarbageCollect().Remove(OnPostGarbageCollectHandle);
+		OnPostGarbageCollectHandle.Reset();
 	}
 
+	bPerformedInitialLoad = false;
+
+	AnimationData.CloseFileReader();
+
 	FrameFileData.Empty();
 	RecordingMaxFrames = 0;
 	MaxFrameDiskSize = 0;
@@ -176,6 +228,11 @@ void ULiveLinkUAssetRecording::WaitForBufferedFrames(int32 InMinFrame, int32 InM
 
 TRange<int32> ULiveLinkUAssetRecording::GetBufferedFrames() const
 {
+	if (bIsFullyLoaded)
+	{
+		return TRange<int32>(0, RecordingMaxFrames);
+	}
+	
 	FScopeLock Lock(&BufferedFrameMutex);
 	if (FrameFileData.Num() > 0)
 	{
@@ -237,6 +294,27 @@ void ULiveLinkUAssetRecording::CopyRecordingData(FLiveLinkPlaybackTracks& InOutL
 	}
 }
 
+void ULiveLinkUAssetRecording::InitializeNewRecordingData(FLiveLinkUAssetRecordingData&& InRecordingData, double InRecordingLengthSeconds)
+{
+	RecordingData = MoveTemp(InRecordingData);
+	LengthInSeconds = InRecordingLengthSeconds;
+	FrameRate = FApp::GetTimecodeFrameRate();
+	RecordingPreset->BuildFromClient();
+
+	int32 MaxFrames = 0;
+	for (const TTuple<FLiveLinkSubjectKey, FLiveLinkRecordingBaseDataContainer>& FrameData : RecordingData.FrameData)
+	{
+		if (FrameData.Value.Timestamps.Num() > MaxFrames)
+		{
+			MaxFrames = FrameData.Value.Timestamps.Num();
+		}
+	}
+
+	RecordingMaxFrames = MaxFrames;
+	
+	bIsFullyLoaded = true;
+}
+
 void ULiveLinkUAssetRecording::SetBufferedFrames(FFrameFileData& InFrameData, const TRange<int32>& InNewRange)
 {
 	FScopeLock Lock(&BufferedFrameMutex);
@@ -245,66 +323,85 @@ void ULiveLinkUAssetRecording::SetBufferedFrames(FFrameFileData& InFrameData, co
 
 void ULiveLinkUAssetRecording::SaveFrameData(FArchive* InFileWriter, const FLiveLinkSubjectKey& InSubjectKey, FLiveLinkRecordingBaseDataContainer& InBaseDataContainer)
 {
-		// This will crash if it fails -- we don't want to save invalid data.
-		InBaseDataContainer.ValidateData();
+	// This will crash if it fails -- we don't want to save invalid data.
+	InBaseDataContainer.ValidateData();
 		
-		// Start block with map key.
-		FGuid Source = InSubjectKey.Source;
-		FString SubjectName = InSubjectKey.SubjectName.ToString();
-		int32 NumFrames = InBaseDataContainer.RecordedData.Num();
-		*InFileWriter << Source;
-		*InFileWriter << SubjectName;
-		*InFileWriter << NumFrames;
-
-		uint64 SerializedFrameSizePosition = 0;
-		int32 SerializedFrameSize = 0;
-		if (NumFrames > 0)
-		{
-			const UScriptStruct* ScriptStruct = InBaseDataContainer.RecordedData[0]->GetScriptStruct();
-			FString StructTypeName = ScriptStruct->GetPathName();
+	// Start block with map key.
+	FGuid Source = InSubjectKey.Source;
+	FString SubjectName = InSubjectKey.SubjectName.ToString();
+	int32 NumFrames = InBaseDataContainer.RecordedData.Num();
+
+	// We record the frame header size first, so later we can bulk load the entire block into memory, then feed it to a memory reader.
+	const uint64 FrameHeaderSizePosition = InFileWriter->Tell();
+	int32 FrameHeaderSize = 0;
+	*InFileWriter << FrameHeaderSize;
+	const int32 FrameHeaderSizeStart = InFileWriter->Tell();
+	
+	*InFileWriter << Source;
+	*InFileWriter << SubjectName;
+	*InFileWriter << NumFrames;
 
-			// Write the struct name and size so it can be loaded later.
-			*InFileWriter << StructTypeName;
+	if (NumFrames == 0)
+	{
+		UE_LOG(LogLiveLinkHub, Error, TEXT("No frames recorded."));
+		return;
+	}
+	
+	const UScriptStruct* ScriptStruct = InBaseDataContainer.RecordedData[0]->GetScriptStruct();
+	FString StructTypeName = ScriptStruct->GetPathName();
 
-			// Remember the position to write the frame size.
-			SerializedFrameSizePosition = InFileWriter->Tell();
-			*InFileWriter << SerializedFrameSize;
-		}
+	// Write the struct name and size so it can be loaded later.
+	*InFileWriter << StructTypeName;
+
+	// Write the frame header size.
+	{
+		const uint64 CurrentPosition = InFileWriter->Tell();
+		FrameHeaderSize = CurrentPosition - FrameHeaderSizeStart;
+
+		InFileWriter->Seek(FrameHeaderSizePosition);
+		*InFileWriter << FrameHeaderSize;
+		InFileWriter->Seek(CurrentPosition);
+	}
+
+	// Remember the position to write the frame size.
+	const uint64 SerializedFrameSizePosition = InFileWriter->Tell();
+	int32 SerializedFrameSize = 0;
+	*InFileWriter << SerializedFrameSize;
 			
-		for (int32 FrameIdx = 0; FrameIdx < NumFrames; ++FrameIdx)
-		{
-			TSharedPtr<FInstancedStruct>& Frame = InBaseDataContainer.RecordedData[FrameIdx];
-			check(Frame.IsValid() && Frame->IsValid());
+	for (int32 FrameIdx = 0; FrameIdx < NumFrames; ++FrameIdx)
+	{
+		TSharedPtr<FInstancedStruct>& Frame = InBaseDataContainer.RecordedData[FrameIdx];
+		check(Frame.IsValid() && Frame->IsValid());
 			
-			// Write the frame index for streaming frames when loading.
-			*InFileWriter << FrameIdx;
+		// Write the frame index for streaming frames when loading.
+		*InFileWriter << FrameIdx;
 
-			// Write the frame's timestamp.
-			double Timestamp = InBaseDataContainer.Timestamps[FrameIdx];
-			*InFileWriter << Timestamp;
+		// Write the frame's timestamp.
+		double Timestamp = InBaseDataContainer.Timestamps[FrameIdx];
+		*InFileWriter << Timestamp;
 			
-			// Write the entire frame data.
-			uint64 SerializeDataStart = InFileWriter->Tell();
-			FObjectAndNameAsStringProxyArchive StructAr(*InFileWriter, false);
-			Frame->Serialize(StructAr);
+		// Write the entire frame data.
+		uint64 SerializeDataStart = InFileWriter->Tell();
+		FObjectAndNameAsStringProxyArchive StructAr(*InFileWriter, false);
+		Frame->Serialize(StructAr);
 
-			// Store the serialized frame size, so we can write it once later.
-			{
-				int32 CurrentSerializedFrameSize = InFileWriter->Tell() - SerializeDataStart;
-				// Sanity check that the serialized frame size is consistent.
-				ensure(CurrentSerializedFrameSize == SerializedFrameSize || SerializedFrameSize == 0);
-				SerializedFrameSize = CurrentSerializedFrameSize;
-			}
-		}
-
-		if (SerializedFrameSize > 0)
+		// Store the serialized frame size, so we can write it once later.
 		{
-			// Write the frame data offset at the beginning of the block.
-			const uint64 FinalOffset = InFileWriter->Tell();
-			InFileWriter->Seek(SerializedFrameSizePosition);
-			*InFileWriter << SerializedFrameSize;
-			InFileWriter->Seek(FinalOffset);
+			int32 CurrentSerializedFrameSize = InFileWriter->Tell() - SerializeDataStart;
+			// Sanity check that the serialized frame size is consistent.
+			ensure(CurrentSerializedFrameSize == SerializedFrameSize || SerializedFrameSize == 0);
+			SerializedFrameSize = CurrentSerializedFrameSize;
 		}
+	}
+
+	if (SerializedFrameSize > 0)
+	{
+		// Write the frame data offset at the beginning of the block.
+		const uint64 FinalOffset = InFileWriter->Tell();
+		InFileWriter->Seek(SerializedFrameSizePosition);
+		*InFileWriter << SerializedFrameSize;
+		InFileWriter->Seek(FinalOffset);
+	}
 }
 
 static TAutoConsoleVariable<float> CVarLiveLinkHubDebugFrameBufferDelay(
@@ -313,6 +410,12 @@ static TAutoConsoleVariable<float> CVarLiveLinkHubDebugFrameBufferDelay(
 	TEXT("The number of seconds to wait when buffering each frame.")
 );
 
+static TAutoConsoleVariable<int32> CVarFrameBufferUpdate(
+	TEXT("LiveLinkHub.Debug.FrameBufferUpdate"),
+	5,
+	TEXT("The number of frames before updating the buffer status.")
+);
+
 void ULiveLinkUAssetRecording::LoadRecordingAsync(int32 InStartFrame, int32 InCurrentFrame, int32 InNumFramesToLoad)
 {
 	const int32 MaxPossibleFrame = RecordingMaxFrames - 1;
@@ -334,15 +437,16 @@ void ULiveLinkUAssetRecording::LoadRecordingAsync(int32 InStartFrame, int32 InCu
 	}
 
 	DebugSleepTime = CVarLiveLinkHubDebugFrameBufferDelay.GetValueOnAnyThread();
+	ReportFrameBufferOnIteration = CVarFrameBufferUpdate.GetValueOnAnyThread();
 	
 	// Perform initial load and record entry frame file offsets.
 	const bool bInitialLoad = FrameFileData.Num() == 0;
 	if (bInitialLoad)
 	{
-		RecordingFileReader->Seek(0);
+		AnimationData.ResetBulkDataOffset();
 		
 		int32 LoadedRecordingVersion;
-		*RecordingFileReader << LoadedRecordingVersion;
+		AnimationData.ReadBulkDataPrimitive(LoadedRecordingVersion);
 
 		// If we modify the RecordingVersion we can perform import logic here.
 		ensure(LoadedRecordingVersion == RecordingVersion);
@@ -350,19 +454,13 @@ void ULiveLinkUAssetRecording::LoadRecordingAsync(int32 InStartFrame, int32 InCu
 		// Process static data.
 		
 		int32 NumStaticData = 0;
-		*RecordingFileReader << NumStaticData;
+		
+		AnimationData.ReadBulkDataPrimitive(NumStaticData);
 
 		for (int32 StaticIdx = 0; StaticIdx < NumStaticData; ++StaticIdx)
 		{
-			FGuid KeySource = FGuid();
-			FString KeyName;
-			
-			*RecordingFileReader << KeySource;
-			*RecordingFileReader << KeyName;
-
 			// Create framedata just to load initial static frame data. Static data doesn't require this afterward.
 			FFrameFileData TemporaryFrameData;
-			TemporaryFrameData.FrameDataSubjectKey = MakeShared<FLiveLinkSubjectKey>(KeySource, *KeyName);
 			if (!LoadInitialFrameData(TemporaryFrameData))
 			{
 				return;
@@ -375,26 +473,18 @@ void ULiveLinkUAssetRecording::LoadRecordingAsync(int32 InStartFrame, int32 InCu
 		// Process frame data.
 		
 		int32 NumFrameData = 0;
-		*RecordingFileReader << NumFrameData;
+		AnimationData.ReadBulkDataPrimitive(NumFrameData);
 
 		for (int32 FrameIdx = 0; FrameIdx < NumFrameData; ++FrameIdx)
 		{
-			FGuid KeySource = FGuid();
-			FString KeyName;
-			
-			*RecordingFileReader << KeySource;
-			*RecordingFileReader << KeyName;
-
 			FFrameFileData KeyPosition;
-			KeyPosition.FrameDataSubjectKey = MakeShared<FLiveLinkSubjectKey>(KeySource, *KeyName);
 			if (!LoadInitialFrameData(KeyPosition))
 			{
 				return;
 			}
 
 			// Offset to the end of this block if there is multiple NumFrameData.
-			const int64 EndBlockPosition = KeyPosition.GetFrameFilePosition(KeyPosition.MaxFrames);
-			RecordingFileReader->Seek(EndBlockPosition);
+			AnimationData.SetBulkDataOffset(KeyPosition.GetFrameFilePosition(KeyPosition.MaxFrames));
 			FrameFileData.Add(MoveTemp(KeyPosition));
 		}
 	}
@@ -416,8 +506,21 @@ void ULiveLinkUAssetRecording::LoadRecordingAsync(int32 InStartFrame, int32 InCu
 
 bool ULiveLinkUAssetRecording::LoadInitialFrameData(FFrameFileData& OutFrameData)
 {
+	int32 FrameHeaderSize = 0;
+	AnimationData.ReadBulkDataPrimitive(FrameHeaderSize);
+
+	const FLiveLinkHubBulkData::FScopedBulkDataMemoryReader Reader = AnimationData.CreateBulkDataMemoryReader(FrameHeaderSize);
+	
+	FGuid KeySource = FGuid();
+	FString KeyName;
+
+	Reader.GetMemoryReader() << KeySource;
+	Reader.GetMemoryReader() << KeyName;
+			
+	OutFrameData.FrameDataSubjectKey = MakeShared<FLiveLinkSubjectKey>(KeySource, *KeyName);
+	
 	int32 MaxFrames = 0;
-	(*RecordingFileReader) << MaxFrames;
+	Reader.GetMemoryReader() << MaxFrames;
 
 	if (MaxFrames > RecordingMaxFrames)
 	{
@@ -430,12 +533,13 @@ bool ULiveLinkUAssetRecording::LoadInitialFrameData(FFrameFileData& OutFrameData
 	{	
 		FString StructTypeName;
 		int32 SerializedStructureSize;
-			
-		*RecordingFileReader << StructTypeName;
-		*RecordingFileReader << SerializedStructureSize;
+
+		Reader.GetMemoryReader() << StructTypeName;
+
+		AnimationData.ReadBulkDataPrimitive(SerializedStructureSize);
 
 		OutFrameData.SerializedStructureSize = SerializedStructureSize;
-		OutFrameData.RecordingStartFrameFilePosition = RecordingFileReader->Tell();
+		OutFrameData.RecordingStartFrameFilePosition = AnimationData.GetBulkDataOffset();
 
 		OutFrameData.LoadedStruct = FindObject<UScriptStruct>(nullptr, *StructTypeName, true);
 		if (!OutFrameData.LoadedStruct.IsValid())
@@ -477,10 +581,6 @@ void ULiveLinkUAssetRecording::LoadFrameData(FFrameFileData& InFrameData, FLiveL
 			return;
 		}
 
-		// Seek to the requested start frame position in the file.
-		int64 NewPosition = InFrameData.GetFrameFilePosition(RequestedStartFrame);
-		RecordingFileReader->Seek(NewPosition);
-
 		// Arrays to store the newly loaded data.
 		TArray<double> NewTimestamps;
 		TArray<TSharedPtr<FInstancedStruct>> NewRecordedData;
@@ -504,9 +604,23 @@ void ULiveLinkUAssetRecording::LoadFrameData(FFrameFileData& InFrameData, FLiveL
 		{
 			bLoadRight = !bLoadRight;
 		};
+
+		int32 Iteration = 0;
 		
 		while (RightFrameIdx < MaxFrames || LeftFrameIdx >= RequestedStartFrame)
 		{
+			ON_SCOPE_EXIT
+			{
+				Iteration++;
+			};
+			
+			if (bCancelStream)
+			{
+				break;
+			}
+
+			WaitIfPaused_AsyncThread();
+			
 			int32 FrameToLoad;
     
 			if (bLoadRight)
@@ -529,20 +643,7 @@ void ULiveLinkUAssetRecording::LoadFrameData(FFrameFileData& InFrameData, FLiveL
 				FrameToLoad = LeftFrameIdx--;
 				LastLoadedLeftFrame = FrameToLoad;
 			}
-
-			int64 FramePosition = InFrameData.GetFrameFilePosition(FrameToLoad);
-			RecordingFileReader->Seek(FramePosition);
-
-			int32 ParsedFrameIdx = 0;
-			*RecordingFileReader << ParsedFrameIdx;
-
-			// Ensure the parsed frame index matches the expected frame
-			if (!ensure(ParsedFrameIdx == FrameToLoad))
-			{
-				UE_LOG(LogLiveLinkHub, Error, TEXT("Frame index mismatch: expected %d, got %d"), FrameToLoad, ParsedFrameIdx);
-				break;
-			}
-
+			
 			auto InsertFrame = [&](TSharedPtr<FInstancedStruct>& InFrame, double InTimestamp, bool bCopy)
 			{
 #if UE_BUILD_DEBUG
@@ -570,22 +671,36 @@ void ULiveLinkUAssetRecording::LoadFrameData(FFrameFileData& InFrameData, FLiveL
 #endif
 			};
 			
-			// Don't load a frame already in memory.
+			// Don't load a frame already in memory. Do this before reading any bulk data, so we don't continually read from disk.
 			double ExistingTimestamp = 0.f;
-			if (TSharedPtr<FInstancedStruct> ExistingFrame = InDataContainer.TryGetFrame(ParsedFrameIdx, ExistingTimestamp))
+			if (TSharedPtr<FInstancedStruct> ExistingFrame = InDataContainer.TryGetFrame(FrameToLoad, ExistingTimestamp))
 			{
-				check(RecordingFileReader->Tell() != 0);
 				constexpr bool bCopy = true; // Copy, as the data container could still be having its frame data pushed to animation.
 				InsertFrame(ExistingFrame, ExistingTimestamp, bCopy);
 				AlternateLoadDirection();
 				continue;
 			}
 
+			int64 FramePosition = InFrameData.GetFrameFilePosition(FrameToLoad);
+			AnimationData.SetBulkDataOffset(FramePosition);
+			
+			const FLiveLinkHubBulkData::FScopedBulkDataMemoryReader Reader = AnimationData.CreateBulkDataMemoryReader(InFrameData.FrameDiskSize);
+
+			int32 ParsedFrameIdx = 0;
+			Reader.GetMemoryReader() << ParsedFrameIdx;
+
+			// Ensure the parsed frame index matches the expected frame
+			if (!ensure(ParsedFrameIdx == FrameToLoad))
+			{
+				UE_LOG(LogLiveLinkHub, Error, TEXT("Frame index mismatch: expected %d, got %d"), FrameToLoad, ParsedFrameIdx);
+				break;
+			}
+			
 			double Timestamp = 0;
-			*RecordingFileReader << Timestamp;
+			Reader.GetMemoryReader() << Timestamp;
 
 			// Instantiate the animation frame.
-			FObjectAndNameAsStringProxyArchive StructAr(*RecordingFileReader, false);
+			FObjectAndNameAsStringProxyArchive StructAr(Reader.GetMemoryReader(), false);
 			TSharedPtr<FInstancedStruct> DataInstancedStruct = MakeShared<FInstancedStruct>(LoadedStruct);
 			DataInstancedStruct->Serialize(StructAr);
 
@@ -593,6 +708,13 @@ void ULiveLinkUAssetRecording::LoadFrameData(FFrameFileData& InFrameData, FLiveL
 		
 			ensure(NewTimestamps.Num() == NewRecordedData.Num());
 
+			// Update buffered frames so any tasks waiting for a certain buffer amount, as well as the UI, can update.
+			// Record every few frames to avoid constantly locking the mutex.
+			if (Iteration % ReportFrameBufferOnIteration == 0)
+			{
+				SetBufferedFrames(InFrameData, TRange<int32>(LastLoadedLeftFrame, LastLoadedRightFrame));
+			}
+			
 			AlternateLoadDirection();
 
 			// Test slow buffer.
@@ -609,7 +731,11 @@ void ULiveLinkUAssetRecording::LoadFrameData(FFrameFileData& InFrameData, FLiveL
 		}
 
 		// Record all the frames that have been buffered.
-		SetBufferedFrames(InFrameData, TRange<int32>(LastLoadedLeftFrame, LastLoadedRightFrame));
+		const TRange<int32> NewFrameRange(LastLoadedLeftFrame, LastLoadedRightFrame);
+		if (InFrameData.BufferedFrames != NewFrameRange)
+		{
+			SetBufferedFrames(InFrameData, NewFrameRange);
+		}
 
 		// Output the streamed data to the data container. This will unload unused frames.
 		{
@@ -636,6 +762,53 @@ FString ULiveLinkUAssetRecording::GetRecordingDataFilePath() const
 	return AbsoluteFilePath;
 }
 
+void ULiveLinkUAssetRecording::EjectAndUnload()
+{
+	const FLiveLinkHubModule& LiveLinkHubModule = FModuleManager::Get().GetModuleChecked<FLiveLinkHubModule>("LiveLinkHub");
+	if (const TSharedPtr<FLiveLinkHubPlaybackController> Controller = LiveLinkHubModule.GetPlaybackController())
+	{
+		Controller->EjectAndUnload(nullptr, this);
+	}
+}
+
+void ULiveLinkUAssetRecording::WaitIfPaused_AsyncThread()
+{
+	check(!IsInGameThread());
+	
+	if (bPauseStream)
+	{
+		OnStreamPausedEvent->Trigger();
+		OnStreamUnpausedEvent->Wait();
+	}
+}
+
+void ULiveLinkUAssetRecording::PauseStream()
+{
+	if (AsyncStreamTask.IsValid() && !AsyncStreamTask->IsDone())
+	{
+		OnStreamUnpausedEvent->Reset();
+		bPauseStream = true;
+		OnStreamPausedEvent->Wait();
+	}
+}
+
+void ULiveLinkUAssetRecording::UnpauseStream()
+{
+	bPauseStream = false;
+	OnStreamPausedEvent->Reset();
+	OnStreamUnpausedEvent->Trigger();
+}
+
+void ULiveLinkUAssetRecording::OnPreGarbageCollect()
+{
+	PauseStream();
+}
+
+void ULiveLinkUAssetRecording::OnPostGarbageCollect()
+{
+	UnpauseStream();
+}
+
 void ULiveLinkUAssetRecording::FLiveLinkStreamAsyncTask::DoWork()
 {
 	int32 LastStartFrame = -1;
@@ -643,6 +816,8 @@ void ULiveLinkUAssetRecording::FLiveLinkStreamAsyncTask::DoWork()
 	int32 LastInitialFrame = -1;
 	while (LiveLinkRecording && !LiveLinkRecording->bCancelStream)
 	{
+		LiveLinkRecording->WaitIfPaused_AsyncThread();
+		
 		if (LastStartFrame != LiveLinkRecording->EarliestFrameToStream
 			|| LastTotalFrames != LiveLinkRecording->TotalFramesToStream
 			|| LastInitialFrame != LiveLinkRecording->InitialFrameToStream)
diff --git a/Engine/Plugins/Animation/LiveLinkHub/Source/LiveLinkHub/Private/Recording/Implementations/LiveLinkUAssetRecording.h b/Engine/Plugins/Animation/LiveLinkHub/Source/LiveLinkHub/Private/Recording/Implementations/LiveLinkUAssetRecording.h
--- a/Engine/Plugins/Animation/LiveLinkHub/Source/LiveLinkHub/Private/Recording/Implementations/LiveLinkUAssetRecording.h
+++ b/Engine/Plugins/Animation/LiveLinkHub/Source/LiveLinkHub/Private/Recording/Implementations/LiveLinkUAssetRecording.h
@@ -3,6 +3,7 @@
 #pragma once
 
 #include "Async/AsyncWork.h"
+#include "Data/LiveLinkHubBulkData.h"
 #include "StructUtils/InstancedStruct.h"
 #include "LiveLinkPreset.h"
 #include "LiveLinkRole.h"
@@ -14,6 +15,7 @@
 
 #include "LiveLinkUAssetRecording.generated.h"
 
+class FMemoryReader;
 struct FLiveLinkPlaybackTracks;
 
 /** Base data container for a recording track. */
@@ -133,14 +135,20 @@ public:
 	
 	virtual void Serialize(FArchive& Ar) override;
 
+	virtual void PostDuplicate(EDuplicateMode::Type DuplicateMode) override;
+	virtual void PostRename(UObject* OldOuter, const FName OldName) override;
+
+	virtual bool IsFullyLoaded() const override { return bIsFullyLoaded; }
+	virtual bool IsSavingRecordingData() const override { return bIsSavingRecordingData; }
+	
 	/** Save recording data to disk. */
-	void SaveRecording();
+	void SaveRecordingData();
 
 	/** Load recording data from disk. */
-	void LoadRecording(int32 InInitialFrame, int32 InNumFramesToLoad);
+	void LoadRecordingData(int32 InInitialFrame, int32 InNumFramesToLoad);
 
 	/** Free memory and close file reader. */
-	void UnloadRecording();
+	void UnloadRecordingData();
 
 	/** Block until frames are loaded. */
 	void WaitForBufferedFrames(int32 InMinFrame, int32 InMaxFrame);
@@ -156,6 +164,9 @@ public:
 
 	/** Copy the asset's loaded recording data to a format suitable for playback in live link. */
 	void CopyRecordingData(FLiveLinkPlaybackTracks& InOutLiveLinkPlaybackTracks) const;
+
+	/** Initial setup of new recording data. */
+	void InitializeNewRecordingData(FLiveLinkUAssetRecordingData&& InRecordingData, double InRecordingLengthSeconds);
 	
 private:
 	/** Frame data file information when loading from a recording file. */
@@ -201,6 +212,24 @@ private:
 	
 	/** Retrieve the recording data file path for this asset. */
 	FString GetRecordingDataFilePath() const;
+
+	/** Eject this recording and make sure it is unloaded. */
+	void EjectAndUnload();
+
+	/** Make the thread wait if we are paused. */
+	void WaitIfPaused_AsyncThread();
+
+	/** Signal and wait for the stream to be paused. */
+	void PauseStream();
+
+	/** Signal the stream can be resumed. */
+	void UnpauseStream();
+	
+	/** Called before garbage collection. */
+	void OnPreGarbageCollect();
+
+	/** Called after garbage collection. */
+	void OnPostGarbageCollect();
 	
 public:
 	/** Recorded static and frame data. */
@@ -216,6 +245,15 @@ private:
 			LiveLinkRecording = InLiveLinkRecording;
 		}
 
+		~FLiveLinkStreamAsyncTask()
+		{
+			if (LiveLinkRecording)
+			{
+				// Make sure we aren't waiting for a pause.
+				LiveLinkRecording->OnStreamPausedEvent->Trigger();
+			}
+		}
+
 		TStatId GetStatId() const
 		{
 			RETURN_QUICK_DECLARE_CYCLE_STAT(LiveLinkStreamAsyncTask, STATGROUP_ThreadPoolAsyncTasks);
@@ -227,11 +265,9 @@ private:
 		TObjectPtr<ULiveLinkUAssetRecording> LiveLinkRecording;
 	};
 
-	friend class FLiveLinkStreamAsyncTask;
+	/** The animation data -- bulk data stored within this uasset. */
+	FLiveLinkHubBulkData AnimationData;
 	
-	/** The file reader for the recording data. */
-	FArchive* RecordingFileReader = nullptr;
-
 	/** The loaded frame data keys and position. */
 	TArray<FFrameFileData> FrameFileData;
 
@@ -253,9 +289,15 @@ private:
 	/** Signal that the stream should be canceled. */
 	std::atomic<bool> bCancelStream = false;
 
+	/** Signal that the stream should be paused. */
+	std::atomic<bool> bPauseStream = false;
+
 	/** True once a full initial load has been performed -- static + frame data. */
 	std::atomic<bool> bPerformedInitialLoad = false;
 
+	/** If we are currently saving recording frame data to disk. */
+	std::atomic<bool> bIsSavingRecordingData = false;
+
 	/** The maximum frame disk size across frame data. */
 	std::atomic<int32> MaxFrameDiskSize = 0;
 
@@ -268,8 +310,26 @@ private:
 	/** The thread streaming data from disk. */
 	TUniquePtr<FAsyncTask<FLiveLinkStreamAsyncTask>> AsyncStreamTask;
 
+	/** Handle for when gc is about to run. */
+	FDelegateHandle OnPreGarbageCollectHandle;
+
+	/** Handle for when gc has finished. */
+	FDelegateHandle OnPostGarbageCollectHandle;
+
+	/** Signalled when the stream is successfully paused. */
+	FEventRef OnStreamPausedEvent = FEventRef(EEventMode::ManualReset);
+	
+	/** Signalled when the stream has been unpaused. */
+	FEventRef OnStreamUnpausedEvent = FEventRef(EEventMode::ManualReset);
+	
 	/** Test slow frame buffering. */
 	float DebugSleepTime = 0.f;
+
+	/** Write the frame buffer size every n iterations. */
+	int32 ReportFrameBufferOnIteration = 5;
+
+	/** If the recording is fully loaded into memory. */
+	bool bIsFullyLoaded = false;
 	
 	/** The current version of the recording. */
 	const int32 RecordingVersion = 1;
diff --git a/Engine/Plugins/Animation/LiveLinkHub/Source/LiveLinkHub/Private/Recording/Implementations/LiveLinkUAssetRecordingPlayer.cpp b/Engine/Plugins/Animation/LiveLinkHub/Source/LiveLinkHub/Private/Recording/Implementations/LiveLinkUAssetRecordingPlayer.cpp
--- a/Engine/Plugins/Animation/LiveLinkHub/Source/LiveLinkHub/Private/Recording/Implementations/LiveLinkUAssetRecordingPlayer.cpp
+++ b/Engine/Plugins/Animation/LiveLinkHub/Source/LiveLinkHub/Private/Recording/Implementations/LiveLinkUAssetRecordingPlayer.cpp
@@ -343,12 +343,12 @@ FFrameRate FLiveLinkPlaybackTracks::GetInitialFrameRate() const
 
 void FLiveLinkUAssetRecordingPlayer::PreparePlayback(ULiveLinkRecording* CurrentRecording)
 {
-	ULiveLinkUAssetRecording* UAssetRecording = CastChecked<ULiveLinkUAssetRecording>(CurrentRecording);
-	LoadedRecording = UAssetRecording;
-
 	// Ensure there nothing is playing and all settings are default. It's possible the CurrentRecording has settings that need to be cleared,
 	// such as if this was just recorded and is now being loaded.
 	ShutdownPlayback();
+
+	ULiveLinkUAssetRecording* UAssetRecording = CastChecked<ULiveLinkUAssetRecording>(CurrentRecording);
+	LoadedRecording = UAssetRecording;
 	
 	CurrentRecordingPlayback = FLiveLinkPlaybackTracks();
 
@@ -357,16 +357,16 @@ void FLiveLinkUAssetRecordingPlayer::PreparePlayback(ULiveLinkRecording* Current
 
 void FLiveLinkUAssetRecordingPlayer::ShutdownPlayback()
 {
-	if (LoadedRecording)
+	if (LoadedRecording.IsValid())
 	{
-		LoadedRecording->UnloadRecording();
+		LoadedRecording->UnloadRecordingData();
 	}
 }
 
 void FLiveLinkUAssetRecordingPlayer::StreamPlayback(int32 InFromFrame)
 {
 	const int32 InitialFramesToBuffer = GetNumFramesToBuffer();
-	LoadedRecording->LoadRecording(InFromFrame, InitialFramesToBuffer);
+	LoadedRecording->LoadRecordingData(InFromFrame, InitialFramesToBuffer);
 
 	// Make sure there are a few frames ready.
 	LoadedRecording->WaitForBufferedFrames(InFromFrame, InFromFrame + 2);
@@ -376,7 +376,7 @@ void FLiveLinkUAssetRecordingPlayer::StreamPlayback(int32 InFromFrame)
 	const int32 CurrentFramesToBuffer = GetNumFramesToBuffer();
 	if (CurrentFramesToBuffer != InitialFramesToBuffer)
 	{
-		LoadedRecording->LoadRecording(InFromFrame, CurrentFramesToBuffer);
+		LoadedRecording->LoadRecordingData(InFromFrame, CurrentFramesToBuffer);
 	}
 
 	// Take the available recording data.
diff --git a/Engine/Plugins/Animation/LiveLinkHub/Source/LiveLinkHub/Private/Recording/Implementations/LiveLinkUAssetRecordingPlayer.h b/Engine/Plugins/Animation/LiveLinkHub/Source/LiveLinkHub/Private/Recording/Implementations/LiveLinkUAssetRecordingPlayer.h
--- a/Engine/Plugins/Animation/LiveLinkHub/Source/LiveLinkHub/Private/Recording/Implementations/LiveLinkUAssetRecordingPlayer.h
+++ b/Engine/Plugins/Animation/LiveLinkHub/Source/LiveLinkHub/Private/Recording/Implementations/LiveLinkUAssetRecordingPlayer.h
@@ -133,7 +133,7 @@ public:
 
 	virtual TRange<int32> GetBufferedFrames() override
 	{
-		return LoadedRecording ? LoadedRecording->GetBufferedFrames() : TRange<int32>(0, 0);
+		return LoadedRecording.IsValid() ? LoadedRecording->GetBufferedFrames() : TRange<int32>(0, 0);
 	}
 
 private:
@@ -148,5 +148,5 @@ private:
 	FLiveLinkPlaybackTracks CurrentRecordingPlayback;
 
 	/** The recording currently loaded. */
-	TObjectPtr<ULiveLinkUAssetRecording> LoadedRecording;
+	TWeakObjectPtr<ULiveLinkUAssetRecording> LoadedRecording;
 };
diff --git a/Engine/Plugins/Animation/LiveLinkHub/Source/LiveLinkHub/Private/Recording/LiveLinkHubPlaybackController.cpp b/Engine/Plugins/Animation/LiveLinkHub/Source/LiveLinkHub/Private/Recording/LiveLinkHubPlaybackController.cpp
--- a/Engine/Plugins/Animation/LiveLinkHub/Source/LiveLinkHub/Private/Recording/LiveLinkHubPlaybackController.cpp
+++ b/Engine/Plugins/Animation/LiveLinkHub/Source/LiveLinkHub/Private/Recording/LiveLinkHubPlaybackController.cpp
@@ -14,6 +14,10 @@
 #include "HAL/RunnableThread.h"
 #include "HAL/Event.h"
 #include "LiveLinkClient.h"
+#include "LiveLinkHub.h"
+#include "LiveLinkHubModule.h"
+#include "Modules/ModuleManager.h"
+#include "PackageTools.h"
 #include "UI/Widgets/SLiveLinkHubPlaybackWidget.h"
 #include "UObject/Object.h"
 #include "UObject/Package.h"
@@ -62,8 +66,6 @@ private:
 FLiveLinkHubPlaybackController::FLiveLinkHubPlaybackController()
 {
 	Client = &IModularFeatures::Get().GetModularFeature<ILiveLinkClient>(ILiveLinkClient::ModularFeatureName);
-
-	OnSourceRemovedHandle = Client->OnLiveLinkSourceRemoved().AddRaw(this, &FLiveLinkHubPlaybackController::OnSourceRemoved);
 	
 	RecordingPlayer = MakeUnique<FLiveLinkUAssetRecordingPlayer>();
 	Playhead = MakeShared<FLiveLinkHubAtomicQualifiedFrameTime, ESPMode::ThreadSafe>();
@@ -80,11 +82,6 @@ FLiveLinkHubPlaybackController::~FLiveLinkHubPlaybackController()
 		Thread->WaitForCompletion();
 		Thread.Reset();
 	}
-
-	if (OnSourceRemovedHandle.IsValid() && Client)
-	{
-		Client->OnLiveLinkSourceRemoved().Remove(OnSourceRemovedHandle);
-	}
 }
 
 TSharedRef<SWidget> FLiveLinkHubPlaybackController::MakePlaybackWidget()
@@ -183,8 +180,6 @@ void FLiveLinkHubPlaybackController::PreparePlayback(ULiveLinkRecording* InLiveL
 
 				RecordingToPlay.Reset(RecordingStrongPtr.Get());
 				RecordingPlayer->PreparePlayback(RecordingToPlay.Get());
-		
-				CurrentFrameRate = RecordingPlayer->GetInitialFramerate();
 
 				// The start and end of playback.
 				SetSelectionStartTime(FQualifiedFrameTime(FFrameTime::FromDecimal(0.f), GetFrameRate()));
@@ -216,7 +211,7 @@ void FLiveLinkHubPlaybackController::PreparePlayback(ULiveLinkRecording* InLiveL
 
 		if (RecordingToPlay.IsValid())
 		{
-			Eject(PreparePlaybackCallback);
+			EjectAndUnload(PreparePlaybackCallback);
 		}
 		else
 		{
@@ -316,6 +311,12 @@ void FLiveLinkHubPlaybackController::Eject(TFunction<void()> CompletionCallback)
 
 	// Recording is done, clear the pointer.
 	RecordingPlayer->ShutdownPlayback();
+
+	if (RecordingToPlay.IsValid())
+	{
+		// It's possible the initial latent action is still in progress if the user ejected this recording immediately after playing.
+		RecordingToPlay->RecordingPreset->CancelLatentAction();
+	}
 	RecordingToPlay.Reset();
 	
 	if (RollbackPreset.IsValid())
@@ -334,6 +335,63 @@ void FLiveLinkHubPlaybackController::Eject(TFunction<void()> CompletionCallback)
 	}
 }
 
+void FLiveLinkHubPlaybackController::EjectAndUnload(TFunction<void()> EjectCompletionCallback, const ULiveLinkRecording* InRecording)
+{
+	const ULiveLinkRecording* Recording = InRecording ? InRecording : RecordingToPlay.Get();
+	UPackage* Package = Recording ? Recording->GetPackage() : nullptr;
+	const bool bIsSavingRecordingData = Recording && Recording->IsSavingRecordingData();
+
+	// Only eject if this is in reference to the current recording playing.
+	if (InRecording == nullptr || InRecording == RecordingToPlay.Get())
+	{
+		Eject(EjectCompletionCallback);
+	}
+	
+	// Unload on the next tick since this could have been called from multistep operations, such as rename or delete. We need to completely
+	// unload so when loading in the future the bulk data archive will be attached correctly.
+	if (Package && !bIsSavingRecordingData)
+	{
+		constexpr bool bUnloadNextTick = true;
+		UnloadRecordingPackage(Package, bUnloadNextTick);
+	}
+}
+
+void FLiveLinkHubPlaybackController::UnloadRecordingPackage(const TWeakObjectPtr<UPackage>& InPackage, bool bUnloadNextTick)
+{
+	if (!InPackage.IsValid() || PackagesUnloading.Contains(InPackage))
+	{
+		return;
+	}
+	
+	auto UnloadPackage = [PackageToUnload = InPackage, this](double = 0.f, float = 0.f) -> EActiveTimerReturnType
+	{
+		ensure(PackagesUnloading.Remove(PackageToUnload) > 0);
+			
+		if (PackageToUnload.IsValid() && !PackageToUnload->IsDirty() && !PackageToUnload->HasAnyPackageFlags(PKG_IsSaving))
+		{
+			UPackageTools::FUnloadPackageParams UnloadParams({ PackageToUnload.Get() });
+			const bool bUnloaded = UPackageTools::UnloadPackages(UnloadParams);
+			ensure(bUnloaded);
+		}
+		
+		return EActiveTimerReturnType::Stop;
+	};
+
+	PackagesUnloading.Add(InPackage);
+
+	if (bUnloadNextTick)
+	{
+		const TSharedPtr<FLiveLinkHub> LiveLinkHub = FModuleManager::Get().GetModuleChecked<FLiveLinkHubModule>("LiveLinkHub").GetLiveLinkHub();
+		LiveLinkHub->GetRootWindow()->RegisterActiveTimer(
+			0.f,
+			FWidgetActiveTimerDelegate::CreateLambda(UnloadPackage));
+	}
+	else
+	{
+		UnloadPackage();
+	}
+}
+
 void FLiveLinkHubPlaybackController::GoToTime(FQualifiedFrameTime InTime)
 {
 	// Stop needs to occur to restart playback.
@@ -373,13 +431,16 @@ FQualifiedFrameTime FLiveLinkHubPlaybackController::GetLength() const
 
 	const double Length = RecordingToPlay ? RecordingToPlay->LengthInSeconds : 0.f;
 	
-	const double TotalFramesDouble = Length * FrameRate.Numerator;
-	const int32 TotalFrames = FMath::FloorToInt(TotalFramesDouble);
+	// Calculate the exact total number of frames, including fractional frames.
+	const double ExactTotalFrames = Length * FrameRate.AsDecimal();
+	const int32 TotalFrames = FMath::FloorToInt(ExactTotalFrames);
 	
-	const FFrameNumber LastFrameNumber(TotalFrames - 1);
+	// Calculate the fractional part and convert it to subframe precision.
+	const double Subframe = ExactTotalFrames - TotalFrames;
 	
-	const FQualifiedFrameTime FrameTime(LastFrameNumber, FrameRate);
+	const FFrameNumber LastFrameNumber(TotalFrames - 1);
 
+	const FQualifiedFrameTime FrameTime(FFrameTime(LastFrameNumber, Subframe), FrameRate);
 	return FrameTime;
 }
 
@@ -396,7 +457,12 @@ FFrameNumber FLiveLinkHubPlaybackController::GetCurrentFrame() const
 
 FFrameRate FLiveLinkHubPlaybackController::GetFrameRate() const
 {
-	return CurrentFrameRate;
+	if (RecordingToPlay.IsValid())
+	{
+		return RecordingToPlay->FrameRate;
+	}
+
+	return FFrameRate(60, 1);
 }
 
 TRange<int32> FLiveLinkHubPlaybackController::GetBufferedFrames() const
@@ -482,23 +548,6 @@ void FLiveLinkHubPlaybackController::OnPlaybackFinished_Internal()
 	PlaybackFinishedDelegate.Broadcast();
 }
 
-void FLiveLinkHubPlaybackController::OnSourceRemoved(FGuid Guid)
-{
-	if (RecordingToPlay.IsValid() && !bIsPreparingPlayback)
-	{
-		// Look for a source that is for this recording and eject. This can occur if the user presses the trash icon
-		// on the playback source while it is in playback.
-		for (const FLiveLinkSourcePreset& Presets : RecordingToPlay->RecordingPreset->GetSourcePresets())
-		{
-			if (Presets.Guid == Guid)
-			{
-				Eject();
-				break;
-			}
-		}
-	}
-}
-
 void FLiveLinkHubPlaybackController::PushSubjectData(const FLiveLinkRecordedFrame& NextFrame, bool bForceSync)
 {
 	// If we're sending static data
@@ -512,8 +561,6 @@ void FLiveLinkHubPlaybackController::PushSubjectData(const FLiveLinkRecordedFram
 	{
 		FLiveLinkFrameDataStruct FrameDataStruct;
 		FrameDataStruct.InitializeWith(NextFrame.Data.GetScriptStruct(), (FLiveLinkBaseFrameData*)NextFrame.Data.GetMemory());
-
-		CurrentFrameRate = FrameDataStruct.GetBaseData()->MetaData.SceneTime.Rate;
 		
 		if (bForceSync)
 		{
diff --git a/Engine/Plugins/Animation/LiveLinkHub/Source/LiveLinkHub/Private/Recording/LiveLinkHubPlaybackController.h b/Engine/Plugins/Animation/LiveLinkHub/Source/LiveLinkHub/Private/Recording/LiveLinkHubPlaybackController.h
--- a/Engine/Plugins/Animation/LiveLinkHub/Source/LiveLinkHub/Private/Recording/LiveLinkHubPlaybackController.h
+++ b/Engine/Plugins/Animation/LiveLinkHub/Source/LiveLinkHub/Private/Recording/LiveLinkHubPlaybackController.h
@@ -56,6 +56,25 @@ public:
 	/** Stop playback and restore the previous settings. */
 	void Eject(TFunction<void()> CompletionCallback = nullptr);
 
+	/**
+	 * Eject this recording and make sure it is unloaded.
+	 * @param EjectCompletionCallback Callback after ejection occurs.
+	 * @param InRecording The specific recording to eject and unload. If null, the RecordingToPlay is used. 
+	 */
+	void EjectAndUnload(TFunction<void()> EjectCompletionCallback = nullptr, const ULiveLinkRecording* InRecording = nullptr);
+
+	/**
+	 * Completely unload a recording package.
+	 *
+	 * Optionally allow the package to be unloaded on the next tick. This is to help with ensuring the package is fully unloaded, such as if
+	 * the unload was requested during a multistep operation, like renaming the asset. When the asset is saved, the bulk data file archive is
+	 * lost, and a reload is needed to reattach it.
+	 * 
+	 * @param InPackage The package of the recording.
+	 * @param bUnloadNextTick Whether the package should unload on the next tick or now.
+	 */
+	void UnloadRecordingPackage(const TWeakObjectPtr<UPackage>& InPackage, bool bUnloadNextTick = false);
+	
 	/** Go to a specific time. */
 	void GoToTime(FQualifiedFrameTime InTime);
 	
@@ -152,9 +171,6 @@ private:
 
 	/** Handler called when playback is finished on the playback thread. Is responsible for resetting the livelink state to what it was before we started playback. */
 	void OnPlaybackFinished_Internal();
-
-	/** When a source has been removed from Live Link Hub. */
-	void OnSourceRemoved(FGuid Guid);
 	
 	/**
 	 * Send data to the client.
@@ -224,9 +240,6 @@ private:
 	/** The playback selection end time. */
 	FQualifiedFrameTime SelectionEndTime;
 
-	/** Current framerate of the recording, sampled from the latest frame. */
-	FFrameRate CurrentFrameRate;
-
-	/** Delegate handle for when a source is removed. */
-	FDelegateHandle OnSourceRemovedHandle;
+	/** Packages in the process of unloading. */
+	TSet<TWeakObjectPtr<UPackage>> PackagesUnloading;
 };
diff --git a/Engine/Plugins/Animation/LiveLinkHub/Source/LiveLinkHub/Private/Recording/LiveLinkHubRecordingController.cpp b/Engine/Plugins/Animation/LiveLinkHub/Source/LiveLinkHub/Private/Recording/LiveLinkHubRecordingController.cpp
--- a/Engine/Plugins/Animation/LiveLinkHub/Source/LiveLinkHub/Private/Recording/LiveLinkHubRecordingController.cpp
+++ b/Engine/Plugins/Animation/LiveLinkHub/Source/LiveLinkHub/Private/Recording/LiveLinkHubRecordingController.cpp
@@ -37,6 +37,11 @@ bool FLiveLinkHubRecordingController::IsRecording() const
 	return RecorderImplementation->IsRecording();
 }
 
+bool FLiveLinkHubRecordingController::IsSavingRecording(ULiveLinkRecording* InRecording) const
+{
+	return RecorderImplementation->IsSavingRecording(InRecording);
+}
+
 void FLiveLinkHubRecordingController::RecordStaticData(const FLiveLinkSubjectKey& SubjectKey, TSubclassOf<ULiveLinkRole> Role, const FLiveLinkStaticDataStruct& StaticData)
 {
 	RecorderImplementation->RecordStaticData(SubjectKey, Role, StaticData);
diff --git a/Engine/Plugins/Animation/LiveLinkHub/Source/LiveLinkHub/Private/Recording/LiveLinkHubRecordingController.h b/Engine/Plugins/Animation/LiveLinkHub/Source/LiveLinkHub/Private/Recording/LiveLinkHubRecordingController.h
--- a/Engine/Plugins/Animation/LiveLinkHub/Source/LiveLinkHub/Private/Recording/LiveLinkHubRecordingController.h
+++ b/Engine/Plugins/Animation/LiveLinkHub/Source/LiveLinkHub/Private/Recording/LiveLinkHubRecordingController.h
@@ -27,6 +27,9 @@ public:
 	/** Returns whether we're currently recording. */
 	bool IsRecording() const;
 
+	/** Returns whether the recording is currently saving. */
+	bool IsSavingRecording(ULiveLinkRecording* InRecording) const;
+
 	/** Record static data in the current recording. */
 	void RecordStaticData(const FLiveLinkSubjectKey& SubjectKey, TSubclassOf<class ULiveLinkRole> Role, const FLiveLinkStaticDataStruct& StaticData);
 
diff --git a/Engine/Plugins/Animation/LiveLinkHub/Source/LiveLinkHub/Private/Recording/LiveLinkHubRecordingListController.h b/Engine/Plugins/Animation/LiveLinkHub/Source/LiveLinkHub/Private/Recording/LiveLinkHubRecordingListController.h
--- a/Engine/Plugins/Animation/LiveLinkHub/Source/LiveLinkHub/Private/Recording/LiveLinkHubRecordingListController.h
+++ b/Engine/Plugins/Animation/LiveLinkHub/Source/LiveLinkHub/Private/Recording/LiveLinkHubRecordingListController.h
@@ -8,6 +8,7 @@
 #include "LiveLinkHub.h"
 #include "LiveLinkHubLog.h"
 #include "LiveLinkHubPlaybackController.h"
+#include "LiveLinkHubRecordingController.h"
 #include "SLiveLinkHubRecordingListView.h"
 #include "LiveLinkRecording.h"
 #include "Modules/ModuleManager.h"
@@ -33,20 +34,33 @@ public:
 	}
 
 private:
-	/** Handler called when a recording a clicked to start the recording.  */
+	/** Handler called when a recording is clicked which will start the recording. */
 	void OnImportRecording(const FAssetData& AssetData)
 	{
-		UObject* RecordingAssetData = AssetData.GetAsset();
-		if (!RecordingAssetData)
+		if (const TSharedPtr<FLiveLinkHub> HubPtr = LiveLinkHub.Pin())
 		{
-			UE_LOG(LogLiveLinkHub, Warning, TEXT("Failed to import recording %s"), *AssetData.AssetName.ToString());
-			return;
-		}
+			if (HubPtr->GetRecordingController()->IsRecording())
+			{
+				return;
+			}
 
-		ULiveLinkRecording* ImportedRecording = Cast<ULiveLinkRecording>(RecordingAssetData);
-		
-		if (TSharedPtr<FLiveLinkHub> HubPtr = LiveLinkHub.Pin())
-		{
+			UObject* RecordingAssetData = AssetData.GetAsset();
+			if (!RecordingAssetData)
+			{
+				UE_LOG(LogLiveLinkHub, Warning, TEXT("Failed to import recording %s"), *AssetData.AssetName.ToString());
+				return;
+			}
+
+			ULiveLinkRecording* ImportedRecording = CastChecked<ULiveLinkRecording>(RecordingAssetData);
+			
+			if (UE::IsSavingPackage(nullptr) && !ImportedRecording->IsFullyLoaded())
+			{
+				// With async saving we risk triggering checks during StaticFindObjectFast, even if the package we are loading isn't the one
+				// being saved. This won't occur if the recording is fully loaded into memory already.
+				UE_LOG(LogLiveLinkHub, Warning, TEXT("Can't start recording because a package is saving"));
+				return;
+			}
+			
 			HubPtr->GetPlaybackController()->PreparePlayback(ImportedRecording);
 		}
 	}
@@ -55,7 +69,7 @@ private:
 	{
 		if (const TSharedPtr<FLiveLinkHub> HubPtr = LiveLinkHub.Pin())
 		{
-			HubPtr->GetPlaybackController()->Eject();
+			HubPtr->GetPlaybackController()->EjectAndUnload();
 		}
 	}
 
diff --git a/Engine/Plugins/Animation/LiveLinkHub/Source/LiveLinkHub/Private/Recording/LiveLinkRecorder.h b/Engine/Plugins/Animation/LiveLinkHub/Source/LiveLinkHub/Private/Recording/LiveLinkRecorder.h
--- a/Engine/Plugins/Animation/LiveLinkHub/Source/LiveLinkHub/Private/Recording/LiveLinkRecorder.h
+++ b/Engine/Plugins/Animation/LiveLinkHub/Source/LiveLinkHub/Private/Recording/LiveLinkRecorder.h
@@ -5,6 +5,8 @@
 #include "LiveLinkRole.h"
 #include "LiveLinkTypes.h"
 
+class ULiveLinkRecording;
+
 /** Interface for handling serialization of livelink data. */
 class ILiveLinkRecorder
 {
@@ -21,4 +23,6 @@ public:
 	virtual void RecordFrameData(const FLiveLinkSubjectKey& SubjectKey, const FLiveLinkFrameDataStruct& FrameData) = 0;
 	/** Return whether we are currently recording livelink data. */
 	virtual bool IsRecording() const = 0;
+	/** Return whether the recording is in the process of being saved. */
+	virtual bool IsSavingRecording(ULiveLinkRecording* InRecording) const = 0;
 };
diff --git a/Engine/Plugins/Animation/LiveLinkHub/Source/LiveLinkHub/Private/Recording/LiveLinkRecording.h b/Engine/Plugins/Animation/LiveLinkHub/Source/LiveLinkHub/Private/Recording/LiveLinkRecording.h
--- a/Engine/Plugins/Animation/LiveLinkHub/Source/LiveLinkHub/Private/Recording/LiveLinkRecording.h
+++ b/Engine/Plugins/Animation/LiveLinkHub/Source/LiveLinkHub/Private/Recording/LiveLinkRecording.h
@@ -3,6 +3,7 @@
 #pragma once
 
 #include "LiveLinkPreset.h"
+#include "Misc/FrameRate.h"
 #include "UObject/Object.h"
 
 #include "LiveLinkRecording.generated.h"
@@ -18,6 +19,12 @@ public:
 		RecordingPreset = CreateDefaultSubobject<ULiveLinkPreset>(TEXT("RecordingPreset"));
 	}
 
+	/** True if this asset has all data loaded. */
+	virtual bool IsFullyLoaded() const { return false; }
+
+	/** True while recording data is being written to bulk data. */
+	virtual bool IsSavingRecordingData() const { return false; }
+
 	/** LiveLink Preset used to save the initial state of the sources and subjects at the time of recording. */
 	UPROPERTY(Instanced)
 	TObjectPtr<ULiveLinkPreset> RecordingPreset = nullptr;
@@ -25,4 +32,8 @@ public:
 	/** Length of the recording. */
 	UPROPERTY()
 	double LengthInSeconds = 0.0;
+
+	/** The framerate of the recording. */
+	UPROPERTY()
+	FFrameRate FrameRate;
 };
\ No newline at end of file
diff --git a/Engine/Plugins/Animation/LiveLinkHub/Source/LiveLinkHub/Private/UI/Widgets/SLiveLinkHubMainTabView.cpp b/Engine/Plugins/Animation/LiveLinkHub/Source/LiveLinkHub/Private/UI/Widgets/SLiveLinkHubMainTabView.cpp
--- a/Engine/Plugins/Animation/LiveLinkHub/Source/LiveLinkHub/Private/UI/Widgets/SLiveLinkHubMainTabView.cpp
+++ b/Engine/Plugins/Animation/LiveLinkHub/Source/LiveLinkHub/Private/UI/Widgets/SLiveLinkHubMainTabView.cpp
@@ -40,7 +40,7 @@ const FText SLiveLinkHubMainTabView::ClientDetailsTabName = LOCTEXT("ClientDetai
 
 void SLiveLinkHubMainTabView::Construct(const FArguments& InArgs)
 {
-	PanelController = MakeShared<FLiveLinkPanelController>();
+	PanelController = MakeShared<FLiveLinkPanelController>(TAttribute<bool>::CreateSP(this, &SLiveLinkHubMainTabView::IsSourcePanelReadOnly));
 	PanelController->OnSubjectSelectionChanged().AddSP(this, &SLiveLinkHubMainTabView::OnSubjectSelectionChanged);
 
 	SLiveLinkHubTabViewWithManagerBase::Construct(
@@ -189,7 +189,7 @@ TSharedRef<SDockTab> SLiveLinkHubMainTabView::SpawnSourcesTab(const FSpawnTabArg
 				.IsEnabled_Lambda([this]()
 				{
 					const FLiveLinkHubModule& LiveLinkHubModule = FModuleManager::Get().GetModuleChecked<FLiveLinkHubModule>("LiveLinkHub");
-					return !LiveLinkHubModule.GetPlaybackController()->IsInPlayback();
+					return !LiveLinkHubModule.GetPlaybackController()->IsInPlayback() && !LiveLinkHubModule.GetRecordingController()->IsRecording();
 				})
 			]
 			+SVerticalBox::Slot()
@@ -290,4 +290,10 @@ void SLiveLinkHubMainTabView::OnSubjectSelectionChanged(const FLiveLinkSubjectKe
 	LiveLinkHubModule.GetSubjectController()->SetSubject(SubjectKey);
 }
 
+bool SLiveLinkHubMainTabView::IsSourcePanelReadOnly() const
+{
+	FLiveLinkHubModule& LiveLinkHubModule = FModuleManager::Get().GetModuleChecked<FLiveLinkHubModule>("LiveLinkHub");
+	return LiveLinkHubModule.GetPlaybackController()->IsInPlayback() || LiveLinkHubModule.GetRecordingController()->IsRecording();
+}
+
 #undef LOCTEXT_NAMESPACE
diff --git a/Engine/Plugins/Animation/LiveLinkHub/Source/LiveLinkHub/Private/UI/Widgets/SLiveLinkHubMainTabView.h b/Engine/Plugins/Animation/LiveLinkHub/Source/LiveLinkHub/Private/UI/Widgets/SLiveLinkHubMainTabView.h
--- a/Engine/Plugins/Animation/LiveLinkHub/Source/LiveLinkHub/Private/UI/Widgets/SLiveLinkHubMainTabView.h
+++ b/Engine/Plugins/Animation/LiveLinkHub/Source/LiveLinkHub/Private/UI/Widgets/SLiveLinkHubMainTabView.h
@@ -52,6 +52,9 @@ private:
 	/** Handles selection change in the subjects view. */
 	void OnSubjectSelectionChanged(const FLiveLinkSubjectKey& SubjectKey);
 
+	/** If the source panel should prevent source editing. */
+	bool IsSourcePanelReadOnly() const;
+
 private:
 	/** Holds the livelink panel controller responsible for creating sources and subjects tabs. */
 	TSharedPtr<class FLiveLinkPanelController> PanelController;
