LiveLinkHub - Add IP Address to connected client info

diff --git a/Engine/Plugins/Animation/LiveLinkHub/Source/LiveLinkHub/LiveLinkHub.Build.cs b/Engine/Plugins/Animation/LiveLinkHub/Source/LiveLinkHub/LiveLinkHub.Build.cs
--- a/Engine/Plugins/Animation/LiveLinkHub/Source/LiveLinkHub/LiveLinkHub.Build.cs
+++ b/Engine/Plugins/Animation/LiveLinkHub/Source/LiveLinkHub/LiveLinkHub.Build.cs
@@ -32,6 +32,7 @@ public class LiveLinkHub : ModuleRules
 				"LiveLinkHubMessaging",
 				"LiveLinkInterface",
 				"LiveLinkMessageBusFramework",
+				"Messaging",
 				"ProjectSettingsViewer",
 				"SharedSettingsWidgets",
 				"Settings",
diff --git a/Engine/Plugins/Animation/LiveLinkHub/Source/LiveLinkHub/Private/Clients/LiveLinkHubProvider.cpp b/Engine/Plugins/Animation/LiveLinkHub/Source/LiveLinkHub/Private/Clients/LiveLinkHubProvider.cpp
--- a/Engine/Plugins/Animation/LiveLinkHub/Source/LiveLinkHub/Private/Clients/LiveLinkHubProvider.cpp
+++ b/Engine/Plugins/Animation/LiveLinkHub/Source/LiveLinkHub/Private/Clients/LiveLinkHubProvider.cpp
@@ -11,6 +11,7 @@
 #include "Delegates/DelegateCombinations.h"
 #include "Editor.h"
 #include "HAL/CriticalSection.h"
+#include "INetworkMessagingExtension.h"
 #include "LiveLinkHubLog.h"
 #include "LiveLinkHubMessages.h"
 #include "LiveLinkSettings.h"
@@ -25,6 +26,55 @@
 
 #define LOCTEXT_NAMESPACE "LiveLinkHub.LiveLinkHubProvider"
 
+namespace LiveLinkHubProviderUtils
+{
+	static INetworkMessagingExtension* GetMessagingStatistics()
+	{
+		IModularFeatures& ModularFeatures = IModularFeatures::Get();
+
+		if (IsInGameThread())
+		{
+			if (ModularFeatures.IsModularFeatureAvailable(INetworkMessagingExtension::ModularFeatureName))
+			{
+				return &ModularFeatures.GetModularFeature<INetworkMessagingExtension>(INetworkMessagingExtension::ModularFeatureName);
+			}
+		}
+		else
+		{
+			IModularFeatures::FScopedLockModularFeatureList ScopedLockModularFeatureList;
+
+			if (ModularFeatures.IsModularFeatureAvailable(INetworkMessagingExtension::ModularFeatureName))
+			{
+				return &ModularFeatures.GetModularFeature<INetworkMessagingExtension>(INetworkMessagingExtension::ModularFeatureName);
+			}
+		}
+
+
+		ensureMsgf(false, TEXT("Feature %s is unavailable"), *INetworkMessagingExtension::ModularFeatureName.ToString());
+		return nullptr;
+	}
+
+	FString GetIPAddress(const FMessageAddress& ClientAddress)
+	{
+		FString IPAddress;
+		if (INetworkMessagingExtension* Statistics = GetMessagingStatistics())
+		{
+			const FGuid NodeId = Statistics->GetNodeIdFromAddress(ClientAddress);
+			IPAddress = NodeId.IsValid() ? Statistics->GetLatestNetworkStatistics(NodeId).IPv4AsString : FString();
+
+			int32 PortIndex = INDEX_NONE;
+			IPAddress.FindChar(TEXT(':'), PortIndex);
+
+			// Cut off the port from the end.
+			if (PortIndex != INDEX_NONE)
+			{
+				IPAddress.LeftInline(PortIndex);
+			}
+		}
+		return IPAddress;
+	}
+}
+
 
 FLiveLinkHubProvider::FLiveLinkHubProvider(const TSharedRef<ILiveLinkHubSessionManager>& InSessionManager)
 	: FLiveLinkProvider(TEXT("LiveLink Hub"), false)
@@ -177,6 +227,8 @@ void FLiveLinkHubProvider::HandleHubConnectMessage(const FLiveLinkHubConnectMess
 	{
 		// Actually added a new entry in the map.
 		FLiveLinkHubUEClientInfo NewClient{Message.ClientInfo};
+		NewClient.IPAddress = LiveLinkHubProviderUtils::GetIPAddress(ConnectionAddress);
+
 		const FLiveLinkHubClientId NewClientId = NewClient.Id;
 		{
 			FWriteScopeLock Locker(ClientsMapLock);
diff --git a/Engine/Plugins/Animation/LiveLinkHub/Source/LiveLinkHub/Private/Clients/LiveLinkHubUEClientInfo.h b/Engine/Plugins/Animation/LiveLinkHub/Source/LiveLinkHub/Private/Clients/LiveLinkHubUEClientInfo.h
--- a/Engine/Plugins/Animation/LiveLinkHub/Source/LiveLinkHub/Private/Clients/LiveLinkHubUEClientInfo.h
+++ b/Engine/Plugins/Animation/LiveLinkHub/Source/LiveLinkHub/Private/Clients/LiveLinkHubUEClientInfo.h
@@ -88,7 +88,7 @@ struct FLiveLinkHubUEClientInfo
 	UPROPERTY(transient)
 	ELiveLinkClientStatus Status = ELiveLinkClientStatus::Disconnected;
 	
-	UPROPERTY()
+	UPROPERTY(VisibleAnywhere, Category = "Client Details")
 	FString IPAddress;
 	
 	/** Name of the host of the UE client */
