LiveLinkHub - Disallow setting outbound name to nothing

diff --git a/Engine/Plugins/Animation/LiveLinkHub/Source/LiveLinkHub/Private/LiveLinkHubSubjectSettings.cpp b/Engine/Plugins/Animation/LiveLinkHub/Source/LiveLinkHub/Private/LiveLinkHubSubjectSettings.cpp
--- a/Engine/Plugins/Animation/LiveLinkHub/Source/LiveLinkHub/Private/LiveLinkHubSubjectSettings.cpp
+++ b/Engine/Plugins/Animation/LiveLinkHub/Source/LiveLinkHub/Private/LiveLinkHubSubjectSettings.cpp
@@ -63,6 +63,10 @@ void ULiveLinkHubSubjectSettings::PostEditChangeProperty(FPropertyChangedEvent&
 	{
 		if (PreviousOutboundName != *OutboundName)
 		{
+			if (OutboundName.IsEmpty())
+			{
+				OutboundName = PreviousOutboundName.ToString();
+			}
 			NotifyRename();
 		}
 	}
