LiveLinkHub - Fix possible crash on editor shutdown.

diff --git a/Engine/Plugins/Animation/LiveLinkHub/Source/LiveLinkHub/Private/Recording/Implementations/LiveLinkUAssetRecording.cpp b/Engine/Plugins/Animation/LiveLinkHub/Source/LiveLinkHub/Private/Recording/Implementations/LiveLinkUAssetRecording.cpp
--- a/Engine/Plugins/Animation/LiveLinkHub/Source/LiveLinkHub/Private/Recording/Implementations/LiveLinkUAssetRecording.cpp
+++ b/Engine/Plugins/Animation/LiveLinkHub/Source/LiveLinkHub/Private/Recording/Implementations/LiveLinkUAssetRecording.cpp
@@ -22,7 +22,10 @@
 
 ULiveLinkUAssetRecording::~ULiveLinkUAssetRecording()
 {
-	UnloadRecordingData();
+	if (!IsTemplate() && !IsEngineExitRequested())
+	{
+		UnloadRecordingData();
+	}
 }
 
 void ULiveLinkUAssetRecording::Serialize(FArchive& Ar)
@@ -139,7 +142,11 @@ void ULiveLinkUAssetRecording::LoadRecordingData(int32 InInitialFrame, int32 InN
 
 void ULiveLinkUAssetRecording::UnloadRecordingData()
 {
-	if (IsSavingRecordingData() || GetPackage()->HasAnyPackageFlags(PKG_IsSaving))
+	// We need to prevent unloading if a package is being saved, but if this is called in a case where there is no outer,
+	// the engine will CastCheck to find the package and fail. We need to avoid in this scenario since that would imply
+	// no package is being saved. This was reported being triggered during an editor shutdown under certain conditions.
+	const UPackage* Package = (GetOuter() == nullptr) ? nullptr : GetPackage();
+	if (IsSavingRecordingData() || (Package && Package->HasAnyPackageFlags(PKG_IsSaving)))
 	{
 		UE_LOG(LogLiveLinkHub, Warning, TEXT("Attempted to unload %s while the package was still being saved"), *GetName());
 		return;
