LiveLinkHub - Fix config data saving and loading issues.

- Fix saving over existing configs not recognizing changes.
- Fix devices that take more than one frame to shutdown not being cleaned up immediately when opening a new config.
- Fix canceling opening a config clearing out the current session.

diff --git a/Engine/Plugins/Animation/LiveLink/Source/LiveLink/Private/LiveLinkClient.cpp b/Engine/Plugins/Animation/LiveLink/Source/LiveLink/Private/LiveLinkClient.cpp
--- a/Engine/Plugins/Animation/LiveLink/Source/LiveLink/Private/LiveLinkClient.cpp
+++ b/Engine/Plugins/Animation/LiveLink/Source/LiveLink/Private/LiveLinkClient.cpp
@@ -556,6 +556,40 @@ void FLiveLinkClient::RemoveAllSources()
 	});
 }
 
+bool FLiveLinkClient::RemoveAllSourcesWithTimeout(float InTimeout)
+{
+	RemoveAllSources();
+
+	const double MaxTime = FPlatformTime::Seconds() + InTimeout;
+
+	auto GetNumNonDefaultSources = [this]()
+	{
+		int32 NumNonDefaultSources = 0;
+		Collection->ForEachSource([&NumNonDefaultSources] (const FLiveLinkCollectionSourceItem& SourceItem)
+		{
+			if (SourceItem.Guid != FLiveLinkSourceCollection::DefaultVirtualSubjectGuid)
+			{
+				NumNonDefaultSources++;
+			}
+		});
+
+		return NumNonDefaultSources;
+	};
+	
+	while (GetNumNonDefaultSources() > 0)
+	{
+		const double CurrentTime = FPlatformTime::Seconds();
+		if (CurrentTime >= MaxTime)
+		{
+			return false;
+		}
+
+		FPlatformProcess::Sleep(0.002f);
+	}
+
+	return true;
+}
+
 bool FLiveLinkClient::HasSourceBeenAdded(TSharedPtr<ILiveLinkSource> InSource) const
 {
 	check(Collection);
diff --git a/Engine/Plugins/Animation/LiveLink/Source/LiveLink/Public/LiveLinkClient.h b/Engine/Plugins/Animation/LiveLink/Source/LiveLink/Public/LiveLinkClient.h
--- a/Engine/Plugins/Animation/LiveLink/Source/LiveLink/Public/LiveLinkClient.h
+++ b/Engine/Plugins/Animation/LiveLink/Source/LiveLink/Public/LiveLinkClient.h
@@ -150,6 +150,13 @@ public:
 	/** Remove all sources from the live link client */
 	void RemoveAllSources();
 
+	/**
+	 * Remove all sources and wait for them to be removed. This is a blocking operation.
+	 * @param InTimeout The timeout in seconds to wait.
+	 * @return Whether all sources were removed successfully.
+	 */
+	bool RemoveAllSourcesWithTimeout(float InTimeout);
+
 #if WITH_EDITOR
 	/** Call initialize again on an existing virtual subject. Used for when a Blueprint Virtual Subject is compiled */
 	void ReinitializeVirtualSubject(const FLiveLinkSubjectKey& SubjectKey);
diff --git a/Engine/Plugins/Animation/LiveLinkHub/Source/LiveLinkHub/Private/Session/LiveLinkHubSessionManager.h b/Engine/Plugins/Animation/LiveLinkHub/Source/LiveLinkHub/Private/Session/LiveLinkHubSessionManager.h
--- a/Engine/Plugins/Animation/LiveLinkHub/Source/LiveLinkHub/Private/Session/LiveLinkHubSessionManager.h
+++ b/Engine/Plugins/Animation/LiveLinkHub/Source/LiveLinkHub/Private/Session/LiveLinkHubSessionManager.h
@@ -10,6 +10,8 @@
 #include "HAL/CriticalSection.h"
 #include "IDesktopPlatform.h"
 #include "LiveLinkHubClient.h"
+#include "LiveLinkHubLog.h"
+#include "Settings/LiveLinkHubSettings.h"
 
 #define LOCTEXT_NAMESPACE "LiveLinkHub.SessionManager"
 
@@ -138,11 +140,14 @@ public:
 		TSharedPtr<FLiveLinkHubSession> CurrentSessionPtr;
 		{
 			FScopeLock Lock(&CurrentSessionCS);
-			CurrentSessionPtr =  CurrentSession;
+			CurrentSessionPtr = CurrentSession;
 		}
 
-		ULiveLinkHubSessionData* LiveLinkHubSessionData = Cast<ULiveLinkHubSessionData>(StaticCastSharedPtr<FLiveLinkHubSession>(CurrentSessionPtr)->SessionData.Get());
+		ULiveLinkHubSessionData* LiveLinkHubSessionData = CastChecked<ULiveLinkHubSessionData>(StaticCastSharedPtr<FLiveLinkHubSession>(CurrentSessionPtr)->SessionData.Get());
 
+		LiveLinkHubSessionData->Sources.Empty();
+		LiveLinkHubSessionData->Subjects.Empty();
+		
 		TArray<FGuid> SourceGuids = LiveLinkHubClient->GetSources();
 		for (const FGuid& SourceGuid : SourceGuids)
 		{
@@ -182,8 +187,6 @@ public:
 
 		const FString DefaultFile = UE::LiveLinkHub::FileUtilities::Private::ConfigDefaultFileName;
 
-		ClearSession();
-
 		TArray<FString> OpenFileNames;
 
 		IDesktopPlatform* DesktopPlatform = FDesktopPlatformModule::Get();
@@ -199,6 +202,10 @@ public:
 
 		if (bFileSelected && OpenFileNames.Num() > 0)
 		{
+			// Certain sources may take time to clean up. If they don't complete in time then the new config being loaded may not create
+			// duplicate sources correctly. There should be errors in the logs of the sources that failed to remove or were unable to be added.
+			constexpr bool bWaitForSourceRemoval = true;
+			ClearSession(bWaitForSourceRemoval);
 			RestoreSession(OpenFileNames[0]);
 		}
 	}
@@ -263,16 +270,19 @@ private:
 	}
 
 	/** Clear the hub data contained in the current session, resetting the hub to its default state. */
-	void ClearSession()
+	void ClearSession(bool bWaitForSourceRemoval = false)
 	{
 		FLiveLinkHubClient* LiveLinkHubClient = static_cast<FLiveLinkHubClient*>(&IModularFeatures::Get().GetModularFeature<ILiveLinkClient>(ILiveLinkClient::ModularFeatureName));
 		check(LiveLinkHubClient);
-
-		LiveLinkHubClient->RemoveAllSources();
-
-		// Removing sources only sets bPendingKill to true, need to tick to make sure they are removed.
-		LiveLinkHubClient->ForceTick();
-
+		
+		const float TimeToWaitForRemoval = bWaitForSourceRemoval ? GetDefault<ULiveLinkHubSettings>()->SourceMaxCleanupTime : 0.f;
+		const bool bRemovedAllSources = LiveLinkHubClient->RemoveAllSourcesWithTimeout(TimeToWaitForRemoval);
+		
+		if (!bRemovedAllSources && bWaitForSourceRemoval)
+		{
+			UE_LOG(LogLiveLinkHub, Warning, TEXT("Could not remove all existing sources in time. Sources may still be getting cleaned up."));
+		}
+		
 		TSharedPtr<FLiveLinkHubSession> CurrentSessionPtr;
 		{
 			FScopeLock Lock(&CurrentSessionCS);
diff --git a/Engine/Plugins/Animation/LiveLinkHub/Source/LiveLinkHub/Private/Settings/LiveLinkHubSettings.h b/Engine/Plugins/Animation/LiveLinkHub/Source/LiveLinkHub/Private/Settings/LiveLinkHubSettings.h
--- a/Engine/Plugins/Animation/LiveLinkHub/Source/LiveLinkHub/Private/Settings/LiveLinkHubSettings.h
+++ b/Engine/Plugins/Animation/LiveLinkHub/Source/LiveLinkHub/Private/Settings/LiveLinkHubSettings.h
@@ -50,6 +50,10 @@ public:
 	UPROPERTY(config, EditAnywhere, Category="LiveLinkHub", meta = (ConfigRestartRequired = true, ClampMin="15.0"))
 	float TargetFrameRate = 60.0f;
 
+	/** Maximum time in seconds to wait for sources to clean up. */
+	UPROPERTY(config, EditAnywhere, AdvancedDisplay, Category="LiveLinkHub", meta = (ClampMin="0.0"))
+	float SourceMaxCleanupTime = 0.25f;
+	
 	/** The filename template to use when creating recordings. */
 	UPROPERTY(config, EditAnywhere, Category="Templates")
 	FString FilenameTemplate = TEXT("NewLiveLinkRecording");
