Fix livelinkhub not connecting to NDisplay render nodes

diff --git a/Engine/Plugins/Animation/LiveLinkHub/Source/LiveLinkHubMessaging/Private/LiveLinkHubConnectionManager.h b/Engine/Plugins/Animation/LiveLinkHub/Source/LiveLinkHubMessaging/Private/LiveLinkHubConnectionManager.h
--- a/Engine/Plugins/Animation/LiveLinkHub/Source/LiveLinkHubMessaging/Private/LiveLinkHubConnectionManager.h
+++ b/Engine/Plugins/Animation/LiveLinkHub/Source/LiveLinkHubMessaging/Private/LiveLinkHubConnectionManager.h
@@ -17,6 +17,8 @@
 #include "Misc/CoreDelegates.h"
 #include "Modules/ModuleManager.h"
 #include "TimerManager.h"
+#include "UObject/UObjectGlobals.h"
+
 
 #if WITH_EDITOR
 #include "Editor.h"
@@ -46,11 +48,14 @@ class FLiveLinkHubConnectionManager
 public:
 	FLiveLinkHubConnectionManager()
 	{
+		FCoreUObjectDelegates::PostLoadMapWithWorld.AddRaw(this, &FLiveLinkHubConnectionManager::PostLoadMap);
 		FCoreDelegates::OnPostEngineInit.AddRaw(this, &FLiveLinkHubConnectionManager::StartDiscovery);
 	}
 
 	~FLiveLinkHubConnectionManager()
 	{
+		FCoreUObjectDelegates::PostLoadMapWithWorld.RemoveAll(this);
+
 		if (FTimerManager* TimerManager = GetTimerManager())
 		{
 			TimerManager->ClearTimer(ConnectionUpdateTimer);
@@ -66,11 +71,13 @@ private:
 	/** Add a discovery request and start polling for results. */
 	void StartDiscovery()
 	{
-		ILiveLinkModule::Get().GetMessageBusDiscoveryManager().AddDiscoveryMessageRequest();
-
-		if (FTimerManager* TimerManager = GetTimerManager())
+		if (!ConnectionUpdateTimer.IsValid())
 		{
-			TimerManager->SetTimer(ConnectionUpdateTimer, FTimerDelegate::CreateRaw(this, &FLiveLinkHubConnectionManager::LookForLiveLinkHubConnection), GetDefault<ULiveLinkSettings>()->MessageBusPingRequestFrequency, true);
+			if (FTimerManager* TimerManager = GetTimerManager())
+			{
+				TimerManager->SetTimer(ConnectionUpdateTimer, FTimerDelegate::CreateRaw(this, &FLiveLinkHubConnectionManager::LookForLiveLinkHubConnection), GetDefault<ULiveLinkSettings>()->MessageBusPingRequestFrequency, true);
+				ILiveLinkModule::Get().GetMessageBusDiscoveryManager().AddDiscoveryMessageRequest();
+			}
 		}
 	}
 
@@ -160,6 +167,12 @@ private:
 			UE_LOG(LogLiveLinkHubConnectionManager, Warning, TEXT("LiveLink modular feature was unavailable."));
 		}
 	}
+
+	/** Handler called when a map changes, used to register the ConnectionUpdateTimer. */
+	void PostLoadMap(UWorld*)
+	{
+		StartDiscovery();
+	}
 	
 private:
 	/** Handle to the timer used to check for livelink hub providers. */
