LiveLinkHub - Use a registry key to find the livelinkhub executable

diff --git a/Engine/Plugins/Animation/LiveLinkHub/Source/LiveLinkHub/Private/LiveLinkHubModule.cpp b/Engine/Plugins/Animation/LiveLinkHub/Source/LiveLinkHub/Private/LiveLinkHubModule.cpp
--- a/Engine/Plugins/Animation/LiveLinkHub/Source/LiveLinkHub/Private/LiveLinkHubModule.cpp
+++ b/Engine/Plugins/Animation/LiveLinkHub/Source/LiveLinkHub/Private/LiveLinkHubModule.cpp
@@ -15,8 +15,6 @@
 #include "Settings/LiveLinkHubSettingsCustomization.h"
 #include "Subjects/LiveLinkHubSubjectSettingsDetailsCustomization.h"
 #include "HAL/FileManager.h"
-#include "Misc/AsyncTaskNotification.h"
-#include "ToolMenus.h"
 
 #define LOCTEXT_NAMESPACE "LiveLinkHubModule"
 
@@ -46,16 +44,6 @@ void FLiveLinkHubModule::ShutdownLiveLinkHub()
 
 void FLiveLinkHubModule::StartupModule()
 {
-	FToolMenuOwnerScoped OwnerScoped(this);
-	UToolMenu* Menu = UToolMenus::Get()->ExtendMenu("LevelEditor.MainMenu.Tools");
-	FToolMenuSection& Section = Menu->AddSection("VirtualProductionSection", LOCTEXT("VirtualProductionSection", "Virtual Production"));
-
-	Section.AddMenuEntry("LiveLinkHub",
-		LOCTEXT("LiveLinkHubLabel", "LiveLink Hub"),
-		LOCTEXT("LiveLinkHubTooltip", "Launch the LiveLink Hub app."),
-		FSlateIcon("LiveLinkStyle", "LiveLinkClient.Common.Icon.Small"),
-		FUIAction(FExecuteAction::CreateRaw(this, &FLiveLinkHubModule::OpenLiveLinkHub)));
-
 	FPropertyEditorModule& PropertyModule = FModuleManager::LoadModuleChecked<FPropertyEditorModule>("PropertyEditor");
 
 	PropertyModule.RegisterCustomClassLayout(ULiveLinkHubSettings::StaticClass()->GetFName(), FOnGetDetailCustomizationInstance::CreateStatic(&FLiveLinkHubSettingsCustomization::MakeInstance));
@@ -70,8 +58,6 @@ void FLiveLinkHubModule::StartupModule()
 
 void FLiveLinkHubModule::ShutdownModule()
 {
-	UToolMenus::UnregisterOwner(this);
-
 	if (FPropertyEditorModule* PropertyEditorModule = FModuleManager::GetModulePtr<FPropertyEditorModule>("PropertyEditor"))
 	{
 		PropertyEditorModule->UnregisterCustomClassLayout(ULiveLinkHubSettings::StaticClass()->GetFName());
@@ -108,75 +94,6 @@ TSharedPtr<FLiveLinkHubPlaybackController> FLiveLinkHubModule::GetPlaybackContro
 	return LiveLinkHub ? LiveLinkHub->PlaybackController : nullptr;
 }
 
-void FLiveLinkHubModule::OpenLiveLinkHub() const
-{
-	FAsyncTaskNotificationConfig NotificationConfig;
-	NotificationConfig.bKeepOpenOnFailure = true;
-	NotificationConfig.TitleText = LOCTEXT("LaunchingLiveLinkHub", "Launching LiveLink Hub...");
-	NotificationConfig.LogCategory = &LogLiveLinkHub;
-
-	FAsyncTaskNotification Notification(NotificationConfig);
-
-	// Find livelink hub executable location for our build configuration
-	FString LiveLinkHubPath = FPlatformProcess::GenerateApplicationPath(TEXT("LiveLinkHub"), FApp::GetBuildConfiguration());
-
-	// Validate it exists and fall back to development if it doesn't.
-	if (!IFileManager::Get().FileExists(*LiveLinkHubPath))
-	{
-		LiveLinkHubPath = FPlatformProcess::GenerateApplicationPath(TEXT("LiveLinkHub"), EBuildConfiguration::Development);
-
-		// If it still doesn't exist, fall back to the shipping executable.
-		if (!IFileManager::Get().FileExists(*LiveLinkHubPath))
-		{
-			LiveLinkHubPath = FPlatformProcess::GenerateApplicationPath(TEXT("LiveLinkHub"), EBuildConfiguration::Shipping);
-		}
-	}
-
-	const FText LaunchLiveLinkHubErrorTitle = LOCTEXT("LaunchLiveLinkHubErrorTitle", "Failed to Launch LiveLinkhub.");
-	if (!IFileManager::Get().FileExists(*LiveLinkHubPath))
-	{
-		Notification.SetComplete(
-			LaunchLiveLinkHubErrorTitle,
-			LOCTEXT("LaunchLiveLinkHubError_ExecutableMissing", "Could not find the executable. Have you compiled the LiveLink Hub app?"),
-			false
-			);
-
-		return;
-	}
-
-	// Validate we do not have it running locally
-	const FString AppName = FPaths::GetCleanFilename(LiveLinkHubPath);
-	if (FPlatformProcess::IsApplicationRunning(*AppName))
-	{
-		Notification.SetComplete(
-			LaunchLiveLinkHubErrorTitle,
-			LOCTEXT("LaunchLiveLinkHubError_AlreadyRunning", "A LiveLinkHub instance is already running."),
-			false
-			);
-		return;
-	}
-
-	constexpr bool bLaunchDetached = true;
-	constexpr bool bLaunchHidden = false;
-	constexpr bool bLaunchReallyHidden = false;
-
-	const FProcHandle ProcHandle = FPlatformProcess::CreateProc(*LiveLinkHubPath, TEXT(""), bLaunchDetached, bLaunchHidden, bLaunchReallyHidden, nullptr, 0, nullptr, nullptr, nullptr);
-	if (ProcHandle.IsValid())
-	{
-		Notification.SetComplete(
-			LOCTEXT("LaunchedLiveLinkHub", "Launched LiveLink Hub"), FText(), true);
-
-		return;
-	}
-	else // Very unlikely in practice, but possible in theory.
-	{
-		Notification.SetComplete(
-			LaunchLiveLinkHubErrorTitle,
-			LOCTEXT("LaunchLiveLinkHubError_InvalidHandle", "Failed to create the LiveLink Hub process."),
-			false);
-	}
-}
-
 TSharedPtr<FLiveLinkHubSubjectController> FLiveLinkHubModule::GetSubjectController() const
 {
 	return LiveLinkHub ? LiveLinkHub->SubjectController : nullptr;
diff --git a/Engine/Plugins/Animation/LiveLinkHub/Source/LiveLinkHub/Private/LiveLinkHubModule.h b/Engine/Plugins/Animation/LiveLinkHub/Source/LiveLinkHub/Private/LiveLinkHubModule.h
--- a/Engine/Plugins/Animation/LiveLinkHub/Source/LiveLinkHub/Private/LiveLinkHubModule.h
+++ b/Engine/Plugins/Animation/LiveLinkHub/Source/LiveLinkHub/Private/LiveLinkHubModule.h
@@ -14,7 +14,6 @@ class FLiveLinkHubSubjectController;
 class FUICommandList;
 class ILiveLinkHubSessionManager;
 
-
 class FLiveLinkHubModule : public ILiveLinkHubModule
 {
 public:
@@ -41,9 +40,6 @@ public:
 	/** Get the subject controller. */
     TSharedPtr<ILiveLinkHubSessionManager> GetSessionManager() const;
 
-	/** Launch livelink hub. */
-	void OpenLiveLinkHub() const;
-
 private:
 	/** LiveLinkHub object responsible for initializing the different controllers. */
 	TSharedPtr<FLiveLinkHub> LiveLinkHub;
diff --git a/Engine/Plugins/Animation/LiveLinkHub/Source/LiveLinkHubEditor/Private/LiveLinkHubEditorModule.cpp b/Engine/Plugins/Animation/LiveLinkHub/Source/LiveLinkHubEditor/Private/LiveLinkHubEditorModule.cpp
--- a/Engine/Plugins/Animation/LiveLinkHub/Source/LiveLinkHubEditor/Private/LiveLinkHubEditorModule.cpp
+++ b/Engine/Plugins/Animation/LiveLinkHub/Source/LiveLinkHubEditor/Private/LiveLinkHubEditorModule.cpp
@@ -2,6 +2,11 @@
 
 #include "LiveLinkHubEditorModule.h"
 
+#include "HAL/FileManager.h"
+#include "LiveLinkHubEditorSettings.h"
+#include "Misc/App.h"
+#include "Misc/AsyncTaskNotification.h"
+#include "Runtime/Launch/Resources/Version.h"
 #include "SLiveLinkHubEditorStatusBar.h"
 #include "ToolMenus.h"
 
@@ -10,6 +15,103 @@ static TAutoConsoleVariable<int32> CVarLiveLinkHubEnableStatusBar(
 	TEXT("Whether to enable showing the livelink hub status bar in the editor. Must be set before launching the editor."),
 	ECVF_RenderThreadSafe);
 
+DECLARE_LOG_CATEGORY_CLASS(LogLiveLinkHubEditor, Log, Log)
+
+#define LOCTEXT_NAMESPACE "LiveLinkHubEditor"
+
+#if PLATFORM_WINDOWS
+#include "Windows/AllowWindowsPlatformTypes.h"
+THIRD_PARTY_INCLUDES_START
+#include <winreg.h>
+THIRD_PARTY_INCLUDES_END
+#include "Windows/HideWindowsPlatformTypes.h"
+#endif
+
+#if DETECT_LIVELINKHUB
+static const FString LiveLinkHubRegistryPath = TEXT("Software\\Epic Games\\LiveLinkHub\\") VERSION_STRINGIFY(ENGINE_MAJOR_VERSION) TEXT(".") VERSION_STRINGIFY(ENGINE_MINOR_VERSION);
+static const FString LiveLinkHubExecutablePath = TEXT("ExecutablePath");
+
+namespace LiveLinkHubUtils
+{
+	/** Attempt to open the livelinkhub registry key that contains the path to the executable. */
+	HKEY OpenLiveLinkHubKey()
+	{
+		HKEY HkcuRunKey = nullptr;
+
+		const LSTATUS OpenResult = RegCreateKeyEx(HKEY_CURRENT_USER, *LiveLinkHubRegistryPath, 0, NULL, REG_OPTION_NON_VOLATILE, KEY_ALL_ACCESS, NULL, &HkcuRunKey, NULL);
+		if (OpenResult != ERROR_SUCCESS)
+		{
+			UE_LOG(LogLiveLinkHubEditor, Log, TEXT("Error opening registry key %s (%08X)"), *LiveLinkHubRegistryPath, OpenResult);
+			return nullptr;
+		}
+		return HkcuRunKey;
+	}
+
+	/** Set the livelinkhub registry key to the livelinkhub executable path. */
+	bool SaveExecutablePathToRegistry()
+	{
+		HKEY HkcuRunKey = OpenLiveLinkHubKey();
+		if (!HkcuRunKey)
+		{
+			return false;
+		}
+
+		ON_SCOPE_EXIT
+		{
+			RegCloseKey(HkcuRunKey);
+		};
+
+		TCHAR ExeFilePath[MAX_PATH + 1];
+		int32 PathLength = ::GetModuleFileNameW(NULL, ExeFilePath, MAX_PATH + 1);
+		if (PathLength > 0)
+		{
+			const LSTATUS SetResult = RegSetValueEx(HkcuRunKey, *LiveLinkHubExecutablePath, 0, REG_SZ, reinterpret_cast<const BYTE*>(ExeFilePath), (PathLength + 1) * sizeof(TCHAR));
+			if (SetResult != ERROR_SUCCESS)
+			{
+				UE_LOG(LogLiveLinkHubEditor, Error, TEXT("Error setting registry value %s (%08X)"), *LiveLinkHubExecutablePath, SetResult);
+				return false;
+			}
+		}
+		else
+		{
+			UE_LOG(LogLiveLinkHubEditor, Error, TEXT("Error setting getting the executable path while setting the registry key for livelink hub."));
+			return false;
+		}
+
+		return true;
+	}
+
+	/** Get the executable path from the registry. */
+	bool GetExecutablePathFromRegistry(FString& OutExecutablePath)
+	{
+		HKEY HkcuRunKey = OpenLiveLinkHubKey();
+		if (!HkcuRunKey)
+		{
+			return false;
+		}
+
+		ON_SCOPE_EXIT
+		{
+			RegCloseKey(HkcuRunKey);
+		};
+
+		wchar_t InstallPathStr[MAX_PATH];
+		DWORD InstallPathSize = sizeof(InstallPathStr);
+
+		LSTATUS GetResult = RegGetValue(HkcuRunKey, NULL, *LiveLinkHubExecutablePath, RRF_RT_REG_SZ, NULL, &InstallPathStr, &InstallPathSize);
+		if (GetResult != ERROR_SUCCESS)
+		{
+			UE_LOG(LogLiveLinkHubEditor, Error, TEXT("Error getting registry value %s:\"%s\" (%08X)"), *LiveLinkHubRegistryPath, *LiveLinkHubExecutablePath, GetResult);
+			return false;
+		}
+
+		OutExecutablePath = WCHAR_TO_TCHAR(InstallPathStr);
+
+		return true;
+	}
+}
+#endif
+
 void FLiveLinkHubEditorModule::StartupModule()
 {
 	if (!IsRunningCommandlet() && CVarLiveLinkHubEnableStatusBar.GetValueOnAnyThread())
@@ -20,6 +122,8 @@ void FLiveLinkHubEditorModule::StartupModule()
 
 void FLiveLinkHubEditorModule::ShutdownModule()
 {
+	UToolMenus::UnregisterOwner(this);
+
 	if (!IsRunningCommandlet() && CVarLiveLinkHubEnableStatusBar.GetValueOnAnyThread())
 	{
 		FCoreDelegates::OnPostEngineInit.RemoveAll(this);
@@ -32,6 +136,103 @@ void FLiveLinkHubEditorModule::OnPostEngineInit()
 	if (GEditor)
 	{
 		RegisterLiveLinkHubStatusBar();
+		
+#if DETECT_LIVELINKHUB
+		if (GetDefault<ULiveLinkHubEditorSettings>()->bDetectLiveLinkHubExecutable)
+		{
+			LiveLinkHubUtils::GetExecutablePathFromRegistry(LiveLinkHubExecutablePath);
+		}
+
+		if (GetDefault<ULiveLinkHubEditorSettings>()->bWriteLiveLinkHubRegistryKey)
+		{
+			// Set the executable path registry key if it hasn't been set.
+			LiveLinkHubUtils::SaveExecutablePathToRegistry();
+		}
+#endif
+
+		FToolMenuOwnerScoped OwnerScoped(this);
+		UToolMenu* Menu = UToolMenus::Get()->ExtendMenu("LevelEditor.MainMenu.Tools");
+		FToolMenuSection& Section = Menu->AddSection("VirtualProductionSection", LOCTEXT("VirtualProductionSection", "Virtual Production"));
+
+		Section.AddMenuEntry("LiveLinkHub",
+			LOCTEXT("LiveLinkHubLabel", "LiveLink Hub"),
+			LOCTEXT("LiveLinkHubTooltip", "Launch the LiveLink Hub app."),
+			FSlateIcon("LiveLinkStyle", "LiveLinkClient.Common.Icon.Small"),
+			FUIAction(FExecuteAction::CreateRaw(this, &FLiveLinkHubEditorModule::OpenLiveLinkHub)));
+	}
+}
+
+void FLiveLinkHubEditorModule::OpenLiveLinkHub() const
+{
+	FAsyncTaskNotificationConfig NotificationConfig;
+	NotificationConfig.bKeepOpenOnFailure = true;
+	NotificationConfig.TitleText = LOCTEXT("LaunchingLiveLinkHub", "Launching LiveLink Hub...");
+	NotificationConfig.LogCategory = &LogLiveLinkHubEditor;
+
+	FAsyncTaskNotification Notification(NotificationConfig);
+
+	// Find livelink hub executable location for our build configuration
+	FString LiveLinkHubPath = FPlatformProcess::GenerateApplicationPath(TEXT("LiveLinkHub"), FApp::GetBuildConfiguration());
+
+	// Validate it exists and fall back to development if it doesn't.
+	if (!IFileManager::Get().FileExists(*LiveLinkHubPath))
+	{
+		LiveLinkHubPath = FPlatformProcess::GenerateApplicationPath(TEXT("LiveLinkHub"), EBuildConfiguration::Development);
+
+		// If it still doesn't exist, fall back to the shipping executable.
+		if (!IFileManager::Get().FileExists(*LiveLinkHubPath))
+		{
+			LiveLinkHubPath = FPlatformProcess::GenerateApplicationPath(TEXT("LiveLinkHub"), EBuildConfiguration::Shipping);
+		}
+	}
+
+	const FText LaunchLiveLinkHubErrorTitle = LOCTEXT("LaunchLiveLinkHubErrorTitle", "Failed to Launch LiveLinkhub.");
+	if (!IFileManager::Get().FileExists(*LiveLinkHubPath))
+	{
+		Notification.SetComplete(
+			LaunchLiveLinkHubErrorTitle,
+			LOCTEXT("LaunchLiveLinkHubError_ExecutableMissing", "Could not find the executable. Have you compiled the LiveLink Hub app?"),
+			false
+		);
+
+		return;
+	}
+
+	if (GetDefault<ULiveLinkHubEditorSettings>()->bDetectLiveLinkHubExecutable)
+	{
+		LiveLinkHubPath = LiveLinkHubExecutablePath;
+	}
+
+	// Validate we do not have it running locally
+	const FString AppName = FPaths::GetCleanFilename(LiveLinkHubPath);
+	if (FPlatformProcess::IsApplicationRunning(*AppName))
+	{
+		Notification.SetComplete(
+			LaunchLiveLinkHubErrorTitle,
+			LOCTEXT("LaunchLiveLinkHubError_AlreadyRunning", "A LiveLinkHub instance is already running."),
+			false
+		);
+		return;
+	}
+
+	constexpr bool bLaunchDetached = true;
+	constexpr bool bLaunchHidden = false;
+	constexpr bool bLaunchReallyHidden = false;
+
+	const FProcHandle ProcHandle = FPlatformProcess::CreateProc(*LiveLinkHubPath, TEXT(""), bLaunchDetached, bLaunchHidden, bLaunchReallyHidden, nullptr, 0, nullptr, nullptr, nullptr);
+	if (ProcHandle.IsValid())
+	{
+		Notification.SetComplete(
+			LOCTEXT("LaunchedLiveLinkHub", "Launched LiveLink Hub"), FText(), true);
+
+		return;
+	}
+	else // Very unlikely in practice, but possible in theory.
+	{
+		Notification.SetComplete(
+			LaunchLiveLinkHubErrorTitle,
+			LOCTEXT("LaunchLiveLinkHubError_InvalidHandle", "Failed to create the LiveLink Hub process."),
+			false);
 	}
 }
 
@@ -57,3 +258,5 @@ TSharedRef<SWidget> FLiveLinkHubEditorModule::CreateLiveLinkHubWidget()
 }
 
 IMPLEMENT_MODULE(FLiveLinkHubEditorModule, LiveLinkHubEditor);
+
+#undef LOCTEXT_NAMESPACE /*LiveLinkHubEditor*/ 
\ No newline at end of file
diff --git a/Engine/Plugins/Animation/LiveLinkHub/Source/LiveLinkHubEditor/Private/LiveLinkHubEditorModule.h b/Engine/Plugins/Animation/LiveLinkHub/Source/LiveLinkHubEditor/Private/LiveLinkHubEditorModule.h
--- a/Engine/Plugins/Animation/LiveLinkHub/Source/LiveLinkHubEditor/Private/LiveLinkHubEditorModule.h
+++ b/Engine/Plugins/Animation/LiveLinkHub/Source/LiveLinkHubEditor/Private/LiveLinkHubEditorModule.h
@@ -6,6 +6,11 @@
 #include "Misc/CoreDelegates.h"
 #include "Modules/ModuleInterface.h"
 
+#ifndef DETECT_LIVELINKHUB
+#define DETECT_LIVELINKHUB PLATFORM_WINDOWS
+#endif
+
+
 /**
  * Editor module responsible for creating the hub's status bar.
  */
@@ -23,6 +28,9 @@ private:
 	/** Callback used to register the livelink hub status bar. */
 	void OnPostEngineInit();
 
+	/** Launch the livelinkhub executable. */
+	void OpenLiveLinkHub() const;
+
 	/** Extend the editor's bottom status bar to add the livelink hub widget. */
 	void RegisterLiveLinkHubStatusBar();
 
@@ -31,4 +39,7 @@ private:
 
 	/** Instantiate the livelink hub status bar widget. */
     TSharedRef<class SWidget> CreateLiveLinkHubWidget();
+
+	/** Path to livelink hub. */
+	FString LiveLinkHubExecutablePath;
 };
diff --git a/Engine/Plugins/Animation/LiveLinkHub/Source/LiveLinkHubEditor/Private/LiveLinkHubEditorSettings.h b/Engine/Plugins/Animation/LiveLinkHub/Source/LiveLinkHubEditor/Private/LiveLinkHubEditorSettings.h
new file mode 100644
--- /dev/null
+++ b/Engine/Plugins/Animation/LiveLinkHub/Source/LiveLinkHubEditor/Private/LiveLinkHubEditorSettings.h
@@ -0,0 +1,25 @@
+// Copyright Epic Games, Inc. All Rights Reserved.
+
+#pragma once
+
+#include "UObject/Object.h"
+
+#include "LiveLinkHubEditorSettings.generated.h"
+
+/**
+ * Settings for LiveLinkHub in editor.
+ */
+UCLASS(config = Editor, defaultconfig)
+class LIVELINKHUBEDITOR_API ULiveLinkHubEditorSettings : public UObject
+{
+public:
+	GENERATED_BODY()
+
+	/** Whether to find the livelinkhub executable by looking up a registry key. */
+	UPROPERTY(config)
+	bool bDetectLiveLinkHubExecutable = false;
+
+	/** Whether to find the livelinkhub executable by looking up a registry key. */
+	UPROPERTY(config)
+	bool bWriteLiveLinkHubRegistryKey = false;
+};
\ No newline at end of file
diff --git a/Engine/Source/Programs/LiveLinkHubEditor/Config/DefaultEditor.ini b/Engine/Source/Programs/LiveLinkHubEditor/Config/DefaultEditor.ini
new file mode 100644
--- /dev/null
+++ b/Engine/Source/Programs/LiveLinkHubEditor/Config/DefaultEditor.ini
@@ -0,0 +1,2 @@
+[/Script/LiveLinkHubEditor.LiveLinkHubEditorSettings]
+bWriteLiveLinkHubRegistryKey=True
\ No newline at end of file
