LiveLinkHub - Support duplicating and renaming recordings in list context menu. We also now make sure bulk data is unloaded when unloading a recording.

diff --git a/Engine/Plugins/Animation/LiveLinkHub/Source/LiveLinkHub/Private/Data/LiveLinkHubBulkData.cpp b/Engine/Plugins/Animation/LiveLinkHub/Source/LiveLinkHub/Private/Data/LiveLinkHubBulkData.cpp
--- a/Engine/Plugins/Animation/LiveLinkHub/Source/LiveLinkHub/Private/Data/LiveLinkHubBulkData.cpp
+++ b/Engine/Plugins/Animation/LiveLinkHub/Source/LiveLinkHub/Private/Data/LiveLinkHubBulkData.cpp
@@ -15,7 +15,7 @@ FLiveLinkHubBulkData::FScopedBulkDataMemoryReader::FScopedBulkDataMemoryReader(c
 
 FLiveLinkHubBulkData::~FLiveLinkHubBulkData()
 {
-	CloseFileReader();
+	UnloadBulkData();
 }
 
 void FLiveLinkHubBulkData::CloseFileReader()
@@ -23,6 +23,12 @@ void FLiveLinkHubBulkData::CloseFileReader()
 	RecordingFileReader.Reset();
 }
 
+void FLiveLinkHubBulkData::UnloadBulkData()
+{
+	CloseFileReader();
+	BulkData.UnloadBulkData();
+}
+
 void FLiveLinkHubBulkData::ReadBulkData(const int64 InBytesToRead, uint8* InMemory)
 {
 	BulkDataOffset = ReadBulkDataImpl(BulkDataOffset, InBytesToRead, InMemory);
diff --git a/Engine/Plugins/Animation/LiveLinkHub/Source/LiveLinkHub/Private/Data/LiveLinkHubBulkData.h b/Engine/Plugins/Animation/LiveLinkHub/Source/LiveLinkHub/Private/Data/LiveLinkHubBulkData.h
--- a/Engine/Plugins/Animation/LiveLinkHub/Source/LiveLinkHub/Private/Data/LiveLinkHubBulkData.h
+++ b/Engine/Plugins/Animation/LiveLinkHub/Source/LiveLinkHub/Private/Data/LiveLinkHubBulkData.h
@@ -40,6 +40,9 @@ public:
 
 	/** Close the file reader if it is open. */
 	void CloseFileReader();
+
+	/** Unloads the bulk data. */
+	void UnloadBulkData();
 	
 	/** Read data using sizeof(T). */
 	template<typename T>
diff --git a/Engine/Plugins/Animation/LiveLinkHub/Source/LiveLinkHub/Private/Recording/Implementations/LiveLinkUAssetRecording.cpp b/Engine/Plugins/Animation/LiveLinkHub/Source/LiveLinkHub/Private/Recording/Implementations/LiveLinkUAssetRecording.cpp
--- a/Engine/Plugins/Animation/LiveLinkHub/Source/LiveLinkHub/Private/Recording/Implementations/LiveLinkUAssetRecording.cpp
+++ b/Engine/Plugins/Animation/LiveLinkHub/Source/LiveLinkHub/Private/Recording/Implementations/LiveLinkUAssetRecording.cpp
@@ -179,7 +179,7 @@ void ULiveLinkUAssetRecording::UnloadRecordingData()
 
 	bPerformedInitialLoad = false;
 
-	AnimationData.CloseFileReader();
+	AnimationData.UnloadBulkData();
 
 	FrameFileData.Empty();
 	RecordingMaxFrames = 0;
diff --git a/Engine/Plugins/Animation/LiveLinkHub/Source/LiveLinkHub/Private/Recording/SLiveLinkHubRecordingListView.h b/Engine/Plugins/Animation/LiveLinkHub/Source/LiveLinkHub/Private/Recording/SLiveLinkHubRecordingListView.h
--- a/Engine/Plugins/Animation/LiveLinkHub/Source/LiveLinkHub/Private/Recording/SLiveLinkHubRecordingListView.h
+++ b/Engine/Plugins/Animation/LiveLinkHub/Source/LiveLinkHub/Private/Recording/SLiveLinkHubRecordingListView.h
@@ -4,13 +4,17 @@
 
 #include "AssetRegistry/AssetRegistryModule.h"
 #include "AssetRegistry/IAssetRegistry.h"
+#include "AssetToolsModule.h"
 #include "ContentBrowserModule.h"
 #include "Delegates/DelegateCombinations.h"
+#include "Framework/Commands/GenericCommands.h"
 #include "Framework/MultiBox/MultiBoxBuilder.h"
 #include "IContentBrowserSingleton.h"
 #include "Implementations/LiveLinkUAssetRecording.h"
-#include "LiveLinkHub.h"
+#include "LiveLinkHubModule.h"
+#include "Recording/LiveLinkHubPlaybackController.h"
 #include "Recording/LiveLinkRecording.h"
+#include "UObject/SavePackage.h"
 #include "Widgets/Input/SButton.h"
 #include "Widgets/SCompoundWidget.h"
 #include "Widgets/Text/STextBlock.h"
@@ -206,15 +210,17 @@ private:
 
 			AssetPickerConfig.Filter = MakeAssetFilter();
 			AssetPickerConfig.OnAssetDoubleClicked = FOnAssetSelected::CreateRaw(this, &SLiveLinkHubRecordingListView::OnImportRecording);
+			AssetPickerConfig.OnGetAssetContextMenu = FOnGetAssetContextMenu::CreateRaw(this, &SLiveLinkHubRecordingListView::GetAssetContextMenu);
 		}
 
 		MenuBuilder.BeginSection(NAME_None, LOCTEXT("ImportRecording_MenuSection", "Import Recording"));
 		{
+			AssetPicker = ContentBrowser.CreateAssetPicker(AssetPickerConfig);
 			TSharedRef<SWidget> PresetPicker = SNew(SBox)
 			.VAlign(VAlign_Fill)
 			.HAlign(HAlign_Fill)
 			[
-				ContentBrowser.CreateAssetPicker(AssetPickerConfig)
+				AssetPicker.ToSharedRef()
 			];
 
 			MenuBuilder.AddWidget(PresetPicker, FText(), true, false);
@@ -236,6 +242,83 @@ private:
 		return Filter;
 	}
 
+	TSharedPtr<SWidget> GetAssetContextMenu(const TArray<FAssetData>& SelectedAssets) const
+	{
+		if (SelectedAssets.Num() <= 0)
+		{
+			return nullptr;
+		}
+
+		TWeakObjectPtr<UObject> SelectedAsset = SelectedAssets[0].GetAsset();
+		if (!SelectedAsset.IsValid())
+		{
+			return nullptr;
+		}
+
+		FMenuBuilder MenuBuilder(true, MakeShared<FUICommandList>());
+
+		MenuBuilder.BeginSection(TEXT("Recording"), LOCTEXT("RecordingSectionLabel", "Recording"));
+		{
+			MenuBuilder.AddMenuEntry(
+				LOCTEXT("RenameRecordingLabel", "Rename"),
+				LOCTEXT("RenameRecordingTooltip", "Rename the recording"),
+				FSlateIcon(FAppStyle::GetAppStyleSetName(), "GenericCommands.Rename"),
+				FUIAction(
+					FExecuteAction::CreateLambda([SelectedAsset, this] ()
+					{
+						if (SelectedAsset.IsValid())
+						{
+							const FContentBrowserModule& ContentBrowserModule = FModuleManager::Get().LoadModuleChecked<FContentBrowserModule>(TEXT("ContentBrowser"));
+							ContentBrowserModule.Get().ExecuteRename(AssetPicker);
+						}
+					}),
+					FCanExecuteAction::CreateLambda([] () { return true; })
+				)
+			);
+			
+			MenuBuilder.AddMenuEntry(
+				LOCTEXT("DuplicateRecordingLabel", "Duplicate"),
+				LOCTEXT("DuplicateRecordingTooltip", "Duplicate the recording"),
+				FSlateIcon(FAppStyle::GetAppStyleSetName(), "GenericCommands.Duplicate"),
+				FUIAction(
+					FExecuteAction::CreateLambda([SelectedAsset] ()
+					{
+						if (SelectedAsset.IsValid())
+						{
+							IAssetTools& AssetTools = FModuleManager::GetModuleChecked<FAssetToolsModule>("AssetTools").Get();
+							
+							FString TargetName;
+							FString TargetPackageName;
+							IAssetTools::Get().CreateUniqueAssetName(SelectedAsset->GetOutermost()->GetName(), TEXT("_Copy"), TargetPackageName, TargetName);
+
+							// Duplicate the asset.
+							UObject* NewAsset = AssetTools.DuplicateAsset(TargetName, FPackageName::GetLongPackagePath(TargetPackageName), SelectedAsset.Get());
+							FSavePackageArgs SavePackageArgs;
+							SavePackageArgs.TopLevelFlags = RF_Public | RF_Standalone;
+							SavePackageArgs.Error = GLog;
+
+							// Save the package.
+							const FString PackageFileName = FPackageName::LongPackageNameToFilename(TargetPackageName, FPackageName::GetAssetPackageExtension());
+							UPackage::SavePackage(NewAsset->GetPackage(), NewAsset, *PackageFileName, MoveTemp(SavePackageArgs));
+
+							// Unload the source recording data, as the bulk data would have been fully loaded to duplicate.
+							const FLiveLinkHubModule& LiveLinkHubModule = FModuleManager::Get().GetModuleChecked<FLiveLinkHubModule>("LiveLinkHub");
+							const TStrongObjectPtr<ULiveLinkRecording> PlaybackRecording = LiveLinkHubModule.GetPlaybackController()->GetRecording();
+							if (PlaybackRecording.Get() != SelectedAsset)
+							{
+								CastChecked<ULiveLinkUAssetRecording>(SelectedAsset)->UnloadRecordingData();
+							}
+						}
+					}),
+					FCanExecuteAction::CreateLambda([] () { return true; })
+				)
+			);
+		}
+		MenuBuilder.EndSection();
+
+		return MenuBuilder.MakeWidget();
+	}
+	
 	/** The text to display when no assets are found. */
 	static FText GetNoAssetsWarningText()
 	{
@@ -249,6 +332,8 @@ private:
 	FOnEject OnEjectDelegate;
 	/** Delegate used to determine a recording can be ejected. */
 	FCanEject OnCanEjectDelegate;
+	/** The asset picker used for selecting recordings. */
+	TSharedPtr<SWidget> AssetPicker;
 	
 	/** Handle for when an asset is added to the asset registry. */
 	FDelegateHandle OnAssetAddedHandle;
