LiveLinkHub: Ensure we also disable background throttling in editor target.

diff --git a/Engine/Plugins/Animation/LiveLinkHub/Source/LiveLinkHub/Private/LiveLinkHubApplication.h b/Engine/Plugins/Animation/LiveLinkHub/Source/LiveLinkHub/Private/LiveLinkHubApplication.h
--- a/Engine/Plugins/Animation/LiveLinkHub/Source/LiveLinkHub/Private/LiveLinkHubApplication.h
+++ b/Engine/Plugins/Animation/LiveLinkHub/Source/LiveLinkHub/Private/LiveLinkHubApplication.h
@@ -4,7 +4,6 @@
 
 #include "CoreMinimal.h"
 #include "Commandlets/Commandlet.h"
-#include "Editor/EditorPerformanceSettings.h"
 #include "Engine/Engine.h"
 #include "Framework/Application/SlateApplication.h"
 #include "LiveLinkHub.h"
@@ -24,9 +23,6 @@ DEFINE_LOG_CATEGORY_STATIC(LogLiveLinkHubApplication, Log, All);
 
 void LiveLinkHubLoop(const TSharedPtr<FLiveLinkHub>& LiveLinkHub)
 {
-	// Disable throttling for the hub
-	GetMutableDefault<UEditorPerformanceSettings>()->bThrottleCPUWhenNotForeground = false;
-
 	const bool bTickOnGameThread = GetDefault<ULiveLinkHubSettings>()->bTickOnGameThread;
 
 	check(FSlateApplication::IsInitialized());
diff --git a/Engine/Plugins/Animation/LiveLinkHub/Source/LiveLinkHub/Private/LiveLinkHubModule.cpp b/Engine/Plugins/Animation/LiveLinkHub/Source/LiveLinkHub/Private/LiveLinkHubModule.cpp
--- a/Engine/Plugins/Animation/LiveLinkHub/Source/LiveLinkHub/Private/LiveLinkHubModule.cpp
+++ b/Engine/Plugins/Animation/LiveLinkHub/Source/LiveLinkHub/Private/LiveLinkHubModule.cpp
@@ -3,6 +3,7 @@
 #include "LiveLinkHubModule.h"
 
 #include "Clients/LiveLinkHubProvider.h"
+#include "Editor/EditorPerformanceSettings.h"
 #include "LiveLinkHubApplication.h"
 #include "LiveLinkHubLog.h"
 #include "LiveLinkHubSubjectSettings.h"
@@ -32,6 +33,9 @@ void FLiveLinkHubModule::StartLiveLinkHub()
 
 	LiveLinkHub->Initialize();
 
+	// Disable throttling for the hub
+	GetMutableDefault<UEditorPerformanceSettings>()->bThrottleCPUWhenNotForeground = false;
+
 #if IS_PROGRAM
 	LiveLinkHubLoop(LiveLinkHub);
 #endif
