LiveLinkHub - Fix not being able to save recordings over 2 gb.

diff --git a/Engine/Plugins/Animation/LiveLinkHub/Source/LiveLinkHub/Private/Data/LiveLinkHubBulkData.cpp b/Engine/Plugins/Animation/LiveLinkHub/Source/LiveLinkHub/Private/Data/LiveLinkHubBulkData.cpp
--- a/Engine/Plugins/Animation/LiveLinkHub/Source/LiveLinkHub/Private/Data/LiveLinkHubBulkData.cpp
+++ b/Engine/Plugins/Animation/LiveLinkHub/Source/LiveLinkHub/Private/Data/LiveLinkHubBulkData.cpp
@@ -58,12 +58,12 @@ void FLiveLinkHubBulkData::Serialize(FArchive& Ar, UObject* Owner)
 	BulkData.Serialize(Ar, Owner);
 }
 
-void FLiveLinkHubBulkData::WriteBulkData(FBufferArchive& BufferArchive)
+void FLiveLinkHubBulkData::WriteBulkData(TArray64<uint8>& Data)
 {
 	BulkData.Lock(LOCK_READ_WRITE);
-	BulkData.Realloc(BufferArchive.Num());
-	unsigned char* BulkDataPtr = BulkData.Realloc(BufferArchive.Num());
-	FMemory::Memcpy(BulkDataPtr, BufferArchive.GetData(), BufferArchive.Num());
+	BulkData.Realloc(Data.Num());
+	unsigned char* BulkDataPtr = BulkData.Realloc(Data.Num());
+	FMemory::Memmove(BulkDataPtr, Data.GetData(), Data.Num());
 	BulkData.Unlock();
 }
 
diff --git a/Engine/Plugins/Animation/LiveLinkHub/Source/LiveLinkHub/Private/Data/LiveLinkHubBulkData.h b/Engine/Plugins/Animation/LiveLinkHub/Source/LiveLinkHub/Private/Data/LiveLinkHubBulkData.h
--- a/Engine/Plugins/Animation/LiveLinkHub/Source/LiveLinkHub/Private/Data/LiveLinkHubBulkData.h
+++ b/Engine/Plugins/Animation/LiveLinkHub/Source/LiveLinkHub/Private/Data/LiveLinkHubBulkData.h
@@ -70,7 +70,7 @@ public:
 	void Serialize(FArchive& Ar, UObject* Owner);
 
 	/** Write to the bulk data file. */
-	void WriteBulkData(FBufferArchive& BufferArchive);
+	void WriteBulkData(TArray64<uint8>& Data);
 	
 private:
 	/**
diff --git a/Engine/Plugins/Animation/LiveLinkHub/Source/LiveLinkHub/Private/Recording/Implementations/LiveLinkUAssetRecording.cpp b/Engine/Plugins/Animation/LiveLinkHub/Source/LiveLinkHub/Private/Recording/Implementations/LiveLinkUAssetRecording.cpp
--- a/Engine/Plugins/Animation/LiveLinkHub/Source/LiveLinkHub/Private/Recording/Implementations/LiveLinkUAssetRecording.cpp
+++ b/Engine/Plugins/Animation/LiveLinkHub/Source/LiveLinkHub/Private/Recording/Implementations/LiveLinkUAssetRecording.cpp
@@ -55,7 +55,8 @@ void ULiveLinkUAssetRecording::SaveRecordingData()
 {
 	bIsSavingRecordingData = true;
 	
-	FBufferArchive Archive;
+	TArray64<uint8> Memory;
+	FMemoryWriter64 Archive(Memory);
 	
 	int32 RecordingVersionToSave = RecordingVersion;
 	Archive << RecordingVersionToSave;
@@ -78,10 +79,7 @@ void ULiveLinkUAssetRecording::SaveRecordingData()
 		SaveFrameData(&Archive, FrameData.Key, FrameData.Value);
 	}
 
-	AnimationData.WriteBulkData(Archive);
-	
-	Archive.FlushCache();
-	Archive.Close();
+	AnimationData.WriteBulkData(Memory);
 
 	bIsSavingRecordingData = false;
 }
@@ -506,7 +504,7 @@ void ULiveLinkUAssetRecording::LoadRecordingAsync(int32 InStartFrame, int32 InCu
 		}
 		else
 		{
-			UE_LOG(LogLiveLinkHub, Error, TEXT("FrameDataSubjectKey is missing for file %s."), *GetRecordingDataFilePath());
+			UE_LOG(LogLiveLinkHub, Error, TEXT("FrameDataSubjectKey is missing for recording %s."), *GetName());
 		}
 	}
 }
@@ -758,17 +756,6 @@ void ULiveLinkUAssetRecording::LoadFrameData(FFrameFileData& InFrameData, FLiveL
 	}
 }
 
-FString ULiveLinkUAssetRecording::GetRecordingDataFilePath() const
-{
-	const FString AssetPath = GetPathName();
-	FString ObjectDirectory = FPaths::GetPath(AssetPath);
-	ObjectDirectory.RemoveFromStart(TEXT("/Game"));
-	const FString AbsoluteFolderPath = FPaths::ConvertRelativePathToFull(FPaths::Combine(FPaths::ProjectContentDir(), ObjectDirectory));
-	const FString AbsoluteFilePath = FString::Printf(TEXT("%s/%s.rec"), *AbsoluteFolderPath, *FPaths::GetBaseFilename(AssetPath));
-
-	return AbsoluteFilePath;
-}
-
 void ULiveLinkUAssetRecording::EjectAndUnload()
 {
 	const FLiveLinkHubModule& LiveLinkHubModule = FModuleManager::Get().GetModuleChecked<FLiveLinkHubModule>("LiveLinkHub");
diff --git a/Engine/Plugins/Animation/LiveLinkHub/Source/LiveLinkHub/Private/Recording/Implementations/LiveLinkUAssetRecording.h b/Engine/Plugins/Animation/LiveLinkHub/Source/LiveLinkHub/Private/Recording/Implementations/LiveLinkUAssetRecording.h
--- a/Engine/Plugins/Animation/LiveLinkHub/Source/LiveLinkHub/Private/Recording/Implementations/LiveLinkUAssetRecording.h
+++ b/Engine/Plugins/Animation/LiveLinkHub/Source/LiveLinkHub/Private/Recording/Implementations/LiveLinkUAssetRecording.h
@@ -209,9 +209,6 @@ private:
 	/** Load frame data to a data container. */
 	void LoadFrameData(FFrameFileData& InFrameData, FLiveLinkRecordingBaseDataContainer& InDataContainer,
 		int32 RequestedStartFrame, int32 RequestedInitialFrame, int32 RequestedFramesToLoad);
-	
-	/** Retrieve the recording data file path for this asset. */
-	FString GetRecordingDataFilePath() const;
 
 	/** Eject this recording and make sure it is unloaded. */
 	void EjectAndUnload();
