LiveLinkHub - Go to livelinkhub page when clicking the option in the tools menu and the app is not installed.

diff --git a/Engine/Plugins/Animation/LiveLinkHub/Source/LiveLinkHubEditor/LiveLinkHubEditor.Build.cs b/Engine/Plugins/Animation/LiveLinkHub/Source/LiveLinkHubEditor/LiveLinkHubEditor.Build.cs
--- a/Engine/Plugins/Animation/LiveLinkHub/Source/LiveLinkHubEditor/LiveLinkHubEditor.Build.cs
+++ b/Engine/Plugins/Animation/LiveLinkHub/Source/LiveLinkHubEditor/LiveLinkHubEditor.Build.cs
@@ -17,6 +17,7 @@ namespace UnrealBuildTool.Rules
 			new string[]
 			{
 				"Engine",
+				"LauncherPlatform",
 				"LiveLink",
 				"LiveLinkEditor",
 				"LiveLinkHub",
diff --git a/Engine/Plugins/Animation/LiveLinkHub/Source/LiveLinkHubEditor/Private/LiveLinkHubEditorModule.cpp b/Engine/Plugins/Animation/LiveLinkHub/Source/LiveLinkHubEditor/Private/LiveLinkHubEditorModule.cpp
--- a/Engine/Plugins/Animation/LiveLinkHub/Source/LiveLinkHubEditor/Private/LiveLinkHubEditorModule.cpp
+++ b/Engine/Plugins/Animation/LiveLinkHub/Source/LiveLinkHubEditor/Private/LiveLinkHubEditorModule.cpp
@@ -3,6 +3,8 @@
 #include "LiveLinkHubEditorModule.h"
 
 #include "HAL/FileManager.h"
+#include "ILauncherPlatform.h"
+#include "LauncherPlatformModule.h"
 #include "LiveLinkHubEditorSettings.h"
 #include "Misc/App.h"
 #include "Misc/AsyncTaskNotification.h"
@@ -138,11 +140,6 @@ void FLiveLinkHubEditorModule::OnPostEngineInit()
 		RegisterLiveLinkHubStatusBar();
 		
 #if DETECT_LIVELINKHUB
-		if (GetDefault<ULiveLinkHubEditorSettings>()->bDetectLiveLinkHubExecutable)
-		{
-			LiveLinkHubUtils::GetExecutablePathFromRegistry(LiveLinkHubExecutablePath);
-		}
-
 		if (GetDefault<ULiveLinkHubEditorSettings>()->bWriteLiveLinkHubRegistryKey)
 		{
 			// Set the executable path registry key if it hasn't been set.
@@ -162,7 +159,7 @@ void FLiveLinkHubEditorModule::OnPostEngineInit()
 	}
 }
 
-void FLiveLinkHubEditorModule::OpenLiveLinkHub() const
+void FLiveLinkHubEditorModule::OpenLiveLinkHub()
 {
 	FAsyncTaskNotificationConfig NotificationConfig;
 	NotificationConfig.bKeepOpenOnFailure = true;
@@ -170,6 +167,49 @@ void FLiveLinkHubEditorModule::OpenLiveLinkHub() const
 	NotificationConfig.LogCategory = &LogLiveLinkHubEditor;
 
 	FAsyncTaskNotification Notification(NotificationConfig);
+	const FText LaunchLiveLinkHubErrorTitle = LOCTEXT("LaunchLiveLinkHubErrorTitle", "Failed to Launch LiveLinkhub.");
+
+#if DETECT_LIVELINKHUB
+	// Try getting the livelinkhub app location by reading a registry key.
+	if (GetDefault<ULiveLinkHubEditorSettings>()->bDetectLiveLinkHubExecutable)
+	{
+		LiveLinkHubUtils::GetExecutablePathFromRegistry(LiveLinkHubExecutablePath);
+	}
+#endif
+
+	if (GetDefault<ULiveLinkHubEditorSettings>()->bDetectLiveLinkHubExecutable && LiveLinkHubExecutablePath.IsEmpty())
+	{
+		// LiveLinkHub executable was not successfully detected, so open the launcher page instead.
+		ILauncherPlatform* LauncherPlatform = FLauncherPlatformModule::Get();
+		if (!GetDefault<ULiveLinkHubEditorSettings>()->LiveLinkHubStorePage.IsEmpty())
+		{
+			FOpenLauncherOptions OpenOptions(GetDefault<ULiveLinkHubEditorSettings>()->LiveLinkHubStorePage);
+			if (!LauncherPlatform->OpenLauncher(OpenOptions))
+			{
+				Notification.SetComplete(
+					LaunchLiveLinkHubErrorTitle,
+					LOCTEXT("LaunchLiveLinkHubError_CouldNotOpenLauncher", "Could not find the LiveLink Hub page on the epic games store."),
+					false
+				);
+
+				return;
+			}
+
+			Notification.SetComplete(
+				LaunchLiveLinkHubErrorTitle,
+				LOCTEXT("LaunchLiveLinkHub_LaunchFromStore", "Please download and launch LiveLinkHub from the Epic Games Store."),
+				true
+			);
+
+			// Store page was launched successfully, but we're not currently launching LiveLinkHub through the launcher, so the user will have to launch it once to write the registry key.
+			// Todo: Detect if livelinkhub is present in library and launch it through there.
+			return;
+		}
+		else
+		{
+			UE_LOG(LogLiveLinkHubEditor, Error, TEXT("Could not open LiveLinkHub store page, config was empty."));
+		}
+	}
 
 	// Find livelink hub executable location for our build configuration
 	FString LiveLinkHubPath = FPlatformProcess::GenerateApplicationPath(TEXT("LiveLinkHub"), FApp::GetBuildConfiguration());
@@ -186,7 +226,6 @@ void FLiveLinkHubEditorModule::OpenLiveLinkHub() const
 		}
 	}
 
-	const FText LaunchLiveLinkHubErrorTitle = LOCTEXT("LaunchLiveLinkHubErrorTitle", "Failed to Launch LiveLinkhub.");
 	if (!IFileManager::Get().FileExists(*LiveLinkHubPath))
 	{
 		Notification.SetComplete(
diff --git a/Engine/Plugins/Animation/LiveLinkHub/Source/LiveLinkHubEditor/Private/LiveLinkHubEditorModule.h b/Engine/Plugins/Animation/LiveLinkHub/Source/LiveLinkHubEditor/Private/LiveLinkHubEditorModule.h
--- a/Engine/Plugins/Animation/LiveLinkHub/Source/LiveLinkHubEditor/Private/LiveLinkHubEditorModule.h
+++ b/Engine/Plugins/Animation/LiveLinkHub/Source/LiveLinkHubEditor/Private/LiveLinkHubEditorModule.h
@@ -29,7 +29,7 @@ private:
 	void OnPostEngineInit();
 
 	/** Launch the livelinkhub executable. */
-	void OpenLiveLinkHub() const;
+	void OpenLiveLinkHub();
 
 	/** Extend the editor's bottom status bar to add the livelink hub widget. */
 	void RegisterLiveLinkHubStatusBar();
diff --git a/Engine/Plugins/Animation/LiveLinkHub/Source/LiveLinkHubEditor/Private/LiveLinkHubEditorSettings.h b/Engine/Plugins/Animation/LiveLinkHub/Source/LiveLinkHubEditor/Private/LiveLinkHubEditorSettings.h
--- a/Engine/Plugins/Animation/LiveLinkHub/Source/LiveLinkHubEditor/Private/LiveLinkHubEditorSettings.h
+++ b/Engine/Plugins/Animation/LiveLinkHub/Source/LiveLinkHubEditor/Private/LiveLinkHubEditorSettings.h
@@ -22,4 +22,8 @@ public:
 	/** Whether to find the livelinkhub executable by looking up a registry key. */
 	UPROPERTY(config)
 	bool bWriteLiveLinkHubRegistryKey = false;
+
+	/** Relative URI to the LiveLinkHub page in the epic games launcher. */
+	UPROPERTY(config)
+	FString LiveLinkHubStorePage;
 };
\ No newline at end of file
