Fix display name of IP address in LiveLinkHub's connected client info struct.

diff --git a/Engine/Plugins/Animation/LiveLinkHub/Source/LiveLinkHub/Private/Clients/LiveLinkHubUEClientInfo.h b/Engine/Plugins/Animation/LiveLinkHub/Source/LiveLinkHub/Private/Clients/LiveLinkHubUEClientInfo.h
--- a/Engine/Plugins/Animation/LiveLinkHub/Source/LiveLinkHub/Private/Clients/LiveLinkHubUEClientInfo.h
+++ b/Engine/Plugins/Animation/LiveLinkHub/Source/LiveLinkHub/Private/Clients/LiveLinkHubUEClientInfo.h
@@ -88,7 +88,7 @@ struct FLiveLinkHubUEClientInfo
 	UPROPERTY(transient)
 	ELiveLinkClientStatus Status = ELiveLinkClientStatus::Disconnected;
 	
-	UPROPERTY(VisibleAnywhere, Category = "Client Details")
+	UPROPERTY(VisibleAnywhere, Category = "Client Details", meta = (DisplayName = "IP Address"))
 	FString IPAddress;
 	
 	/** Name of the host of the UE client */
