Fix LiveLinkHub Warnings

- Enable ACLPlugin to avoid warning
- Disable Ooodle texture compression (Would generate a warning because oodle plugin wasn't enabled)
- Fix warning generated by calling ScanPathsSynchronous on a path that wasn't mounted. LiveLinkHub's AppData path is now mounted if we're in a launcher build

diff --git a/Engine/Plugins/Animation/LiveLinkHub/Source/LiveLinkHub/Private/LiveLinkHub.cpp b/Engine/Plugins/Animation/LiveLinkHub/Source/LiveLinkHub/Private/LiveLinkHub.cpp
--- a/Engine/Plugins/Animation/LiveLinkHub/Source/LiveLinkHub/Private/LiveLinkHub.cpp
+++ b/Engine/Plugins/Animation/LiveLinkHub/Source/LiveLinkHub/Private/LiveLinkHub.cpp
@@ -36,7 +36,7 @@ void FLiveLinkHub::Preinitialize()
 	IModularFeatures::Get().RegisterModularFeature(ILiveLinkClient::ModularFeatureName, LiveLinkHubClient.Get());
 }
 
-void FLiveLinkHub::Initialize()
+void FLiveLinkHub::Initialize(bool bLauncherDistribution)
 {
 	TRACE_CPUPROFILER_EVENT_SCOPE(FLiveLinkHub::Initialize);
 
@@ -50,9 +50,12 @@ void FLiveLinkHub::Initialize()
 
 	FModuleManager::Get().LoadModule("Settings");
 
-	IAssetRegistry& AssetRegistry = FModuleManager::LoadModuleChecked<FAssetRegistryModule>(AssetRegistryConstants::ModuleName).Get();
-	const FString FilePath = FPaths::Combine(FPlatformProcess::UserSettingsDir(), *FApp::GetEpicProductIdentifier(), TEXT("LiveLinkHub"), TEXT("Content"));
-	AssetRegistry.ScanPathsSynchronous({ FilePath }, /*bForceRescan=*/ true);
+	if (bLauncherDistribution)
+	{
+		IAssetRegistry& AssetRegistry = FModuleManager::LoadModuleChecked<FAssetRegistryModule>(AssetRegistryConstants::ModuleName).Get();
+		const FString FilePath = FPaths::Combine(FPlatformProcess::UserSettingsDir(), *FApp::GetEpicProductIdentifier(), TEXT("LiveLinkHub"), TEXT("Content"));
+		AssetRegistry.ScanPathsSynchronous({ FilePath }, /*bForceRescan=*/ true);
+	}
 
 	CommandExecutor = MakeUnique<FConsoleCommandExecutor>();
 	IModularFeatures::Get().RegisterModularFeature(IConsoleCommandExecutor::ModularFeatureName(), CommandExecutor.Get());
diff --git a/Engine/Plugins/Animation/LiveLinkHub/Source/LiveLinkHub/Private/LiveLinkHub.h b/Engine/Plugins/Animation/LiveLinkHub/Source/LiveLinkHub/Private/LiveLinkHub.h
--- a/Engine/Plugins/Animation/LiveLinkHub/Source/LiveLinkHub/Private/LiveLinkHub.h
+++ b/Engine/Plugins/Animation/LiveLinkHub/Source/LiveLinkHub/Private/LiveLinkHub.h
@@ -51,8 +51,11 @@ public:
 public:
 	/** First phase initialization; must precede FEngineLoop::Init(). */
 	void Preinitialize();
-	/** Launch the slate application and initialize its components. */
-	void Initialize();
+	/** 
+	 * Launch the slate application and initialize its components. 
+	 * @param bLauncherDistribution Whether the hub should be initialized as a launcher distributed LiveLinkHub.
+	 */
+	void Initialize(bool bLauncherDistribution);
 	/** Tick the hub. */
 	void Tick();
 	/** Get the root window that hosts the hub's slate application. */
diff --git a/Engine/Plugins/Animation/LiveLinkHub/Source/LiveLinkHub/Private/LiveLinkHubModule.cpp b/Engine/Plugins/Animation/LiveLinkHub/Source/LiveLinkHub/Private/LiveLinkHubModule.cpp
--- a/Engine/Plugins/Animation/LiveLinkHub/Source/LiveLinkHub/Private/LiveLinkHubModule.cpp
+++ b/Engine/Plugins/Animation/LiveLinkHub/Source/LiveLinkHub/Private/LiveLinkHubModule.cpp
@@ -26,12 +26,12 @@ void FLiveLinkHubModule::PreinitializeLiveLinkHub()
 	LiveLinkHub->Preinitialize();
 }
 
-void FLiveLinkHubModule::StartLiveLinkHub()
+void FLiveLinkHubModule::StartLiveLinkHub(bool bLauncherDistribution)
 {
 	TRACE_CPUPROFILER_EVENT_SCOPE(StartLiveLinkHub);
 	checkf(LiveLinkHub, TEXT("Ensure PreinitializeLiveLinkHub was called first"));
 
-	LiveLinkHub->Initialize();
+	LiveLinkHub->Initialize(bLauncherDistribution);
 
 	// Disable throttling for the hub
 	GetMutableDefault<UEditorPerformanceSettings>()->bThrottleCPUWhenNotForeground = false;
diff --git a/Engine/Plugins/Animation/LiveLinkHub/Source/LiveLinkHub/Private/LiveLinkHubModule.h b/Engine/Plugins/Animation/LiveLinkHub/Source/LiveLinkHub/Private/LiveLinkHubModule.h
--- a/Engine/Plugins/Animation/LiveLinkHub/Source/LiveLinkHub/Private/LiveLinkHubModule.h
+++ b/Engine/Plugins/Animation/LiveLinkHub/Source/LiveLinkHub/Private/LiveLinkHubModule.h
@@ -19,7 +19,7 @@ class FLiveLinkHubModule : public ILiveLinkHubModule
 public:
 	//~ Begin ILiveLinkHubModule interface
 	virtual void PreinitializeLiveLinkHub() override;
-	virtual void StartLiveLinkHub() override;
+	virtual void StartLiveLinkHub(bool bLauncherDistribution = false) override;
 	virtual void ShutdownLiveLinkHub() override;
 	virtual void StartupModule() override;
 	virtual void ShutdownModule() override;
diff --git a/Engine/Plugins/Animation/LiveLinkHub/Source/LiveLinkHub/Public/ILiveLinkHubModule.h b/Engine/Plugins/Animation/LiveLinkHub/Source/LiveLinkHub/Public/ILiveLinkHubModule.h
--- a/Engine/Plugins/Animation/LiveLinkHub/Source/LiveLinkHub/Public/ILiveLinkHubModule.h
+++ b/Engine/Plugins/Animation/LiveLinkHub/Source/LiveLinkHub/Public/ILiveLinkHubModule.h
@@ -10,7 +10,7 @@ public:
 	virtual void PreinitializeLiveLinkHub() = 0;
 
 	/** Launch the slate application hosting the live link hub. */
-	virtual void StartLiveLinkHub() = 0;
+	virtual void StartLiveLinkHub(bool bLauncherDistribution = false) = 0;
 
 	virtual void ShutdownLiveLinkHub() = 0;
 };
diff --git a/Engine/Source/Programs/LiveLinkHub/Private/LiveLinkHubRun.cpp b/Engine/Source/Programs/LiveLinkHub/Private/LiveLinkHubRun.cpp
--- a/Engine/Source/Programs/LiveLinkHub/Private/LiveLinkHubRun.cpp
+++ b/Engine/Source/Programs/LiveLinkHub/Private/LiveLinkHubRun.cpp
@@ -56,6 +56,8 @@ int32 RunLiveLinkHub(const TCHAR* CommandLine)
 		});
 #endif
 
+	bool bLauncherDistribution = false;
+
 #if !IS_PROGRAM
 	const TCHAR* const DevelopmentProjectPath = TEXT("../../Source/Programs/LiveLinkHubEditor/LiveLinkHubEditor.uproject");
 	const TCHAR* const StagedProjectPath = TEXT("../../../LiveLinkHubEditor/LiveLinkHubEditor.uproject");
@@ -66,10 +68,18 @@ int32 RunLiveLinkHub(const TCHAR* CommandLine)
 	}
 	else if (FPaths::FileExists(StagedProjectPath))
 	{
+		bLauncherDistribution = true;
 		FPaths::SetProjectFilePath(StagedProjectPath);
 	}
 #endif
 
+	if (bLauncherDistribution)
+	{
+		// This is needed to ensure that content generated by the hub will be discoverable (ie. LiveLink recordings)
+		const FString GameContentPath = FPaths::Combine(FPlatformProcess::UserSettingsDir(), *FApp::GetEpicProductIdentifier(), TEXT("LiveLinkHub"), TEXT("Content"));
+		FPackageName::RegisterMountPoint(TEXT("/Game/"), GameContentPath);
+	}
+
 #if IS_PROGRAM
 	// Disable this when going through PreInit to prevent the console window from appearing.
 	GIsSilent = true;
@@ -99,7 +109,7 @@ int32 RunLiveLinkHub(const TCHAR* CommandLine)
 
 		Result = GEngineLoop.Init();
 
-		HubModule.StartLiveLinkHub();
+		HubModule.StartLiveLinkHub(bLauncherDistribution);
 
 		// Hide the splash screen now that everything is ready to go
 		FPlatformSplash::Hide();
diff --git a/Engine/Source/Programs/LiveLinkHubEditor/Config/DefaultEngine.ini b/Engine/Source/Programs/LiveLinkHubEditor/Config/DefaultEngine.ini
--- a/Engine/Source/Programs/LiveLinkHubEditor/Config/DefaultEngine.ini
+++ b/Engine/Source/Programs/LiveLinkHubEditor/Config/DefaultEngine.ini
@@ -33,3 +33,7 @@ DetailsPanel.Style.EnableCardLayout=true
 
 [/Script/BuildSettings.BuildSettings]
 DefaultEditorTarget=LiveLinkHubEditor
+
+; Not needed at the moment, can be removed if we enable the TextureFormatOodle plugin
+[AlternateTextureCompression]
+TextureCompressionFormat=""
\ No newline at end of file
diff --git a/Engine/Source/Programs/LiveLinkHubEditor/Source/LiveLinkHubEditor.Target.cs b/Engine/Source/Programs/LiveLinkHubEditor/Source/LiveLinkHubEditor.Target.cs
--- a/Engine/Source/Programs/LiveLinkHubEditor/Source/LiveLinkHubEditor.Target.cs
+++ b/Engine/Source/Programs/LiveLinkHubEditor/Source/LiveLinkHubEditor.Target.cs
@@ -32,6 +32,7 @@ public class LiveLinkHubEditorTarget : TargetRules
 		// They may be a direct dependency or a dependency of another plugin.
 		EnablePlugins.AddRange(new string[]
 		{
+			"ACLPlugin",
 			"LiveLink",
 			"LiveLinkHub",
 			"LiveLinkCamera",
