Find LiveLinkHub installation by iterating over installed apps by the EGS rather than relying on a registry key.

diff --git a/Engine/Plugins/Animation/LiveLinkHub/Source/LiveLinkHubEditor/LiveLinkHubEditor.Build.cs b/Engine/Plugins/Animation/LiveLinkHub/Source/LiveLinkHubEditor/LiveLinkHubEditor.Build.cs
--- a/Engine/Plugins/Animation/LiveLinkHub/Source/LiveLinkHubEditor/LiveLinkHubEditor.Build.cs
+++ b/Engine/Plugins/Animation/LiveLinkHub/Source/LiveLinkHubEditor/LiveLinkHubEditor.Build.cs
@@ -16,6 +16,7 @@ namespace UnrealBuildTool.Rules
 			PrivateDependencyModuleNames.AddRange(
 			new string[]
 			{
+				"Json",
 				"Engine",
 				"LauncherPlatform",
 				"LiveLink",
@@ -23,6 +24,7 @@ namespace UnrealBuildTool.Rules
 				"LiveLinkHub",
 				"LiveLinkHubMessaging",
 				"LiveLinkInterface",
+				"Serialization",
 				"Slate",
 				"SlateCore",
 				"ToolMenus",
diff --git a/Engine/Plugins/Animation/LiveLinkHub/Source/LiveLinkHubEditor/Private/LiveLinkHubEditorModule.cpp b/Engine/Plugins/Animation/LiveLinkHub/Source/LiveLinkHubEditor/Private/LiveLinkHubEditorModule.cpp
--- a/Engine/Plugins/Animation/LiveLinkHub/Source/LiveLinkHubEditor/Private/LiveLinkHubEditorModule.cpp
+++ b/Engine/Plugins/Animation/LiveLinkHub/Source/LiveLinkHubEditor/Private/LiveLinkHubEditorModule.cpp
@@ -2,16 +2,20 @@
 
 #include "LiveLinkHubEditorModule.h"
 
+#include "DesktopPlatformModule.h"
 #include "HAL/FileManager.h"
 #include "ILauncherPlatform.h"
 #include "LauncherPlatformModule.h"
+#include "LiveLinkHubLauncherUtils.h"
 #include "LiveLinkHubEditorSettings.h"
 #include "Misc/App.h"
 #include "Misc/AsyncTaskNotification.h"
+#include "Misc/MessageDialog.h"
 #include "Runtime/Launch/Resources/Version.h"
 #include "SLiveLinkHubEditorStatusBar.h"
 #include "ToolMenus.h"
 
+
 static TAutoConsoleVariable<int32> CVarLiveLinkHubEnableStatusBar(
 	TEXT("LiveLinkHub.EnableStatusBar"), 1,
 	TEXT("Whether to enable showing the livelink hub status bar in the editor. Must be set before launching the editor."),
@@ -21,99 +25,6 @@ DECLARE_LOG_CATEGORY_CLASS(LogLiveLinkHubEditor, Log, Log)
 
 #define LOCTEXT_NAMESPACE "LiveLinkHubEditor"
 
-#if PLATFORM_WINDOWS
-#include "Windows/AllowWindowsPlatformTypes.h"
-THIRD_PARTY_INCLUDES_START
-#include <winreg.h>
-THIRD_PARTY_INCLUDES_END
-#include "Windows/HideWindowsPlatformTypes.h"
-#endif
-
-#if DETECT_LIVELINKHUB
-static const FString LiveLinkHubRegistryPath = TEXT("Software\\Epic Games\\LiveLinkHub\\") VERSION_STRINGIFY(ENGINE_MAJOR_VERSION) TEXT(".") VERSION_STRINGIFY(ENGINE_MINOR_VERSION);
-static const FString LiveLinkHubExecutablePath = TEXT("ExecutablePath");
-
-namespace LiveLinkHubUtils
-{
-	/** Attempt to open the livelinkhub registry key that contains the path to the executable. */
-	HKEY OpenLiveLinkHubKey()
-	{
-		HKEY HkcuRunKey = nullptr;
-
-		const LSTATUS OpenResult = RegCreateKeyEx(HKEY_CURRENT_USER, *LiveLinkHubRegistryPath, 0, NULL, REG_OPTION_NON_VOLATILE, KEY_ALL_ACCESS, NULL, &HkcuRunKey, NULL);
-		if (OpenResult != ERROR_SUCCESS)
-		{
-			UE_LOG(LogLiveLinkHubEditor, Log, TEXT("Error opening registry key %s (%08X)"), *LiveLinkHubRegistryPath, OpenResult);
-			return nullptr;
-		}
-		return HkcuRunKey;
-	}
-
-	/** Set the livelinkhub registry key to the livelinkhub executable path. */
-	bool SaveExecutablePathToRegistry()
-	{
-		HKEY HkcuRunKey = OpenLiveLinkHubKey();
-		if (!HkcuRunKey)
-		{
-			return false;
-		}
-
-		ON_SCOPE_EXIT
-		{
-			RegCloseKey(HkcuRunKey);
-		};
-
-		TCHAR ExeFilePath[MAX_PATH + 1];
-		int32 PathLength = ::GetModuleFileNameW(NULL, ExeFilePath, MAX_PATH + 1);
-		if (PathLength > 0)
-		{
-			const LSTATUS SetResult = RegSetValueEx(HkcuRunKey, *LiveLinkHubExecutablePath, 0, REG_SZ, reinterpret_cast<const BYTE*>(ExeFilePath), (PathLength + 1) * sizeof(TCHAR));
-			if (SetResult != ERROR_SUCCESS)
-			{
-				UE_LOG(LogLiveLinkHubEditor, Error, TEXT("Error setting registry value %s (%08X)"), *LiveLinkHubExecutablePath, SetResult);
-				return false;
-			}
-		}
-		else
-		{
-			UE_LOG(LogLiveLinkHubEditor, Error, TEXT("Error setting getting the executable path while setting the registry key for livelink hub."));
-			return false;
-		}
-
-		return true;
-	}
-
-	/** Get the executable path from the registry. */
-	bool GetExecutablePathFromRegistry(FString& OutExecutablePath)
-	{
-		HKEY HkcuRunKey = OpenLiveLinkHubKey();
-		if (!HkcuRunKey)
-		{
-			return false;
-		}
-
-		ON_SCOPE_EXIT
-		{
-			RegCloseKey(HkcuRunKey);
-		};
-
-		wchar_t InstallPathStr[MAX_PATH];
-		DWORD InstallPathSize = sizeof(InstallPathStr);
-
-		LSTATUS GetResult = RegGetValue(HkcuRunKey, NULL, *LiveLinkHubExecutablePath, RRF_RT_REG_SZ, NULL, &InstallPathStr, &InstallPathSize);
-		if (GetResult != ERROR_SUCCESS)
-		{
-			UE_LOG(LogLiveLinkHubEditor, Error, TEXT("Error getting registry value %s:\"%s\" (%08X)"), *LiveLinkHubRegistryPath, *LiveLinkHubExecutablePath, GetResult);
-			return false;
-		}
-
-		OutExecutablePath = WCHAR_TO_TCHAR(InstallPathStr);
-
-		return true;
-	}
-}
-#endif
-
 void FLiveLinkHubEditorModule::StartupModule()
 {
 	if (!IsRunningCommandlet() && CVarLiveLinkHubEnableStatusBar.GetValueOnAnyThread())
@@ -126,7 +37,7 @@ void FLiveLinkHubEditorModule::ShutdownModule()
 {
 	UToolMenus::UnregisterOwner(this);
 
-	if (!IsRunningCommandlet() && CVarLiveLinkHubEnableStatusBar.GetValueOnAnyThread())
+	if (!IsRunningCommandlet() && CVarLiveLinkHubEnableStatusBar.GetValueOnAnyThread()) 
 	{
 		FCoreDelegates::OnPostEngineInit.RemoveAll(this);
 		UnregisterLiveLinkHubStatusBar();
@@ -139,14 +50,6 @@ void FLiveLinkHubEditorModule::OnPostEngineInit()
 	{
 		RegisterLiveLinkHubStatusBar();
 		
-#if DETECT_LIVELINKHUB
-		if (GetDefault<ULiveLinkHubEditorSettings>()->bWriteLiveLinkHubRegistryKey)
-		{
-			// Set the executable path registry key if it hasn't been set.
-			LiveLinkHubUtils::SaveExecutablePathToRegistry();
-		}
-#endif
-
 		FToolMenuOwnerScoped OwnerScoped(this);
 		UToolMenu* Menu = UToolMenus::Get()->ExtendMenu("LevelEditor.MainMenu.Tools");
 		FToolMenuSection& Section = Menu->AddSection("VirtualProductionSection", LOCTEXT("VirtualProductionSection", "Virtual Production"));
@@ -169,46 +72,86 @@ void FLiveLinkHubEditorModule::OpenLiveLinkHub()
 	FAsyncTaskNotification Notification(NotificationConfig);
 	const FText LaunchLiveLinkHubErrorTitle = LOCTEXT("LaunchLiveLinkHubErrorTitle", "Failed to Launch LiveLinkhub.");
 
-#if DETECT_LIVELINKHUB
 	// Try getting the livelinkhub app location by reading a registry key.
 	if (GetDefault<ULiveLinkHubEditorSettings>()->bDetectLiveLinkHubExecutable)
 	{
-		LiveLinkHubUtils::GetExecutablePathFromRegistry(LiveLinkHubExecutablePath);
-	}
-#endif
-
-	if (GetDefault<ULiveLinkHubEditorSettings>()->bDetectLiveLinkHubExecutable && LiveLinkHubExecutablePath.IsEmpty())
-	{
-		// LiveLinkHub executable was not successfully detected, so open the launcher page instead.
 		ILauncherPlatform* LauncherPlatform = FLauncherPlatformModule::Get();
-		if (!GetDefault<ULiveLinkHubEditorSettings>()->LiveLinkHubStorePage.IsEmpty())
+
+		UE::LiveLinkHubLauncherUtils::FInstalledApp LiveLinkHubApp;
+		if (UE::LiveLinkHubLauncherUtils::FindLiveLinkHubInstallation(LiveLinkHubApp))
 		{
-			FOpenLauncherOptions OpenOptions(GetDefault<ULiveLinkHubEditorSettings>()->LiveLinkHubStorePage);
+			// Found a LiveLinkHub installation from the launcher, so launch it that way.
+			
+			const FString LaunchLink = TEXT("apps") / LiveLinkHubApp.NamespaceId + TEXT("%3A") + LiveLinkHubApp.ItemId + TEXT("%3A") + LiveLinkHubApp.AppName + TEXT("?action=launch&silent=true");
+			FOpenLauncherOptions OpenOptions(LaunchLink);
 			if (!LauncherPlatform->OpenLauncher(OpenOptions))
 			{
 				Notification.SetComplete(
 					LaunchLiveLinkHubErrorTitle,
-					LOCTEXT("LaunchLiveLinkHubError_CouldNotOpenLauncher", "Could not find the LiveLink Hub page on the epic games store."),
+					LOCTEXT("LaunchLiveLinkHubError_CouldNotOpenLauncher", "Could not launch LiveLinkHub throug the epic games store."),
 					false
 				);
 
-				return;
 			}
-
-			Notification.SetComplete(
-				LaunchLiveLinkHubErrorTitle,
-				LOCTEXT("LaunchLiveLinkHub_LaunchFromStore", "Please download and launch LiveLinkHub from the Epic Games Store."),
-				true
-			);
-
-			// Store page was launched successfully, but we're not currently launching LiveLinkHub through the launcher, so the user will have to launch it once to write the registry key.
-			// Todo: Detect if livelinkhub is present in library and launch it through there.
-			return;
+			else
+			{
+				Notification.SetComplete(
+					LOCTEXT("LiveLinkHubLaunchSuccessTitle", "Launched LiveLinkHub."),
+					LOCTEXT("LaunchLiveLinkHubError_LaunchSuccess", "Launching LiveLinkHub through the Epic Games Store."),
+					true
+				);
+			}
 		}
 		else
 		{
-			UE_LOG(LogLiveLinkHubEditor, Error, TEXT("Could not open LiveLinkHub store page, config was empty."));
+			const EAppReturnType::Type Choice = FMessageDialog::Open(EAppMsgType::YesNo, LOCTEXT("InstallThroughEGS", "LiveLinkHub is not currently installed, do you want to install it through the Epic Games Store?"));
+
+			if (Choice == EAppReturnType::Yes)
+			{
+				// Could not find LiveLinkHub from the launcher. Prompt the user to open the EGS and install it.
+				if (!GetDefault<ULiveLinkHubEditorSettings>()->LiveLinkHubStorePage.IsEmpty())
+				{
+					FOpenLauncherOptions OpenOptions(GetDefault<ULiveLinkHubEditorSettings>()->LiveLinkHubStorePage);
+					if (!LauncherPlatform->OpenLauncher(OpenOptions))
+					{
+						Notification.SetComplete(
+							LaunchLiveLinkHubErrorTitle,
+							LOCTEXT("LaunchLiveLinkHubError_CouldNotOpenLauncher", "Could not find the LiveLink Hub page on the epic games store."),
+							false
+						);
+					}
+					else
+					{
+						Notification.SetComplete(
+							LaunchLiveLinkHubErrorTitle,
+							LOCTEXT("LaunchLiveLinkHub_LaunchFromStore", "Opening Epic Games Store to the LiveLinkHub page."),
+							true
+						);
+					}
+
+				}
+				else
+				{
+					Notification.SetComplete(
+						LaunchLiveLinkHubErrorTitle,
+						LOCTEXT("LaunchLiveLinkHubError_EmptyConfig", "Could not find the LiveLink Hub page on the epic games store, missing configuration for the store page."),
+						false
+					);
+				
+				}
+			}
+			else
+			{
+				Notification.SetComplete(
+					LaunchLiveLinkHubErrorTitle,
+					LOCTEXT("LaunchLiveLinkHub_DidNotLaunchFromStore", "LiveLinkHub could not be launched since it wasn't installed."),
+					false
+				);
+			}
+
 		}
+
+		return;
 	}
 
 	// Find livelink hub executable location for our build configuration
@@ -237,11 +180,6 @@ void FLiveLinkHubEditorModule::OpenLiveLinkHub()
 		return;
 	}
 
-	if (GetDefault<ULiveLinkHubEditorSettings>()->bDetectLiveLinkHubExecutable)
-	{
-		LiveLinkHubPath = LiveLinkHubExecutablePath;
-	}
-
 	// Validate we do not have it running locally
 	const FString AppName = FPaths::GetCleanFilename(LiveLinkHubPath);
 	if (FPlatformProcess::IsApplicationRunning(*AppName))
diff --git a/Engine/Plugins/Animation/LiveLinkHub/Source/LiveLinkHubEditor/Private/LiveLinkHubEditorModule.h b/Engine/Plugins/Animation/LiveLinkHub/Source/LiveLinkHubEditor/Private/LiveLinkHubEditorModule.h
--- a/Engine/Plugins/Animation/LiveLinkHub/Source/LiveLinkHubEditor/Private/LiveLinkHubEditorModule.h
+++ b/Engine/Plugins/Animation/LiveLinkHub/Source/LiveLinkHubEditor/Private/LiveLinkHubEditorModule.h
@@ -6,11 +6,6 @@
 #include "Misc/CoreDelegates.h"
 #include "Modules/ModuleInterface.h"
 
-#ifndef DETECT_LIVELINKHUB
-#define DETECT_LIVELINKHUB PLATFORM_WINDOWS
-#endif
-
-
 /**
  * Editor module responsible for creating the hub's status bar.
  */
diff --git a/Engine/Plugins/Animation/LiveLinkHub/Source/LiveLinkHubEditor/Private/LiveLinkHubEditorSettings.h b/Engine/Plugins/Animation/LiveLinkHub/Source/LiveLinkHubEditor/Private/LiveLinkHubEditorSettings.h
--- a/Engine/Plugins/Animation/LiveLinkHub/Source/LiveLinkHubEditor/Private/LiveLinkHubEditorSettings.h
+++ b/Engine/Plugins/Animation/LiveLinkHub/Source/LiveLinkHubEditor/Private/LiveLinkHubEditorSettings.h
@@ -19,11 +19,15 @@ public:
 	UPROPERTY(config)
 	bool bDetectLiveLinkHubExecutable = false;
 
-	/** Whether to find the livelinkhub executable by looking up a registry key. */
-	UPROPERTY(config)
-	bool bWriteLiveLinkHubRegistryKey = false;
-
 	/** Relative URI to the LiveLinkHub page in the epic games launcher. */
 	UPROPERTY(config)
 	FString LiveLinkHubStorePage;
+
+	/** App name of LiveLinkHub, used to automatically detect if it's installed. */
+	UPROPERTY(config)
+	FString LiveLinkHubAppName;
+
+	/** The version of LiveLinkHub that should be used. If empty, any version will be allowed. */
+	UPROPERTY(config)
+	FString LiveLinkHubTargetVersion;
 };
\ No newline at end of file
diff --git a/Engine/Plugins/Animation/LiveLinkHub/Source/LiveLinkHubEditor/Private/LiveLinkHubLauncherUtils.h b/Engine/Plugins/Animation/LiveLinkHub/Source/LiveLinkHubEditor/Private/LiveLinkHubLauncherUtils.h
new file mode 100644
--- /dev/null
+++ b/Engine/Plugins/Animation/LiveLinkHub/Source/LiveLinkHubEditor/Private/LiveLinkHubLauncherUtils.h
@@ -0,0 +1,95 @@
+// Copyright Epic Games, Inc. All Rights Reserved.
+
+#pragma once
+
+#include "Containers/StringConv.h"
+#include "HAL/FileManager.h"
+#include "LiveLinkHubEditorSettings.h"
+#include "Misc/FileHelper.h"
+#include "Serialization/JsonTypes.h"
+#include "Serialization/JsonReader.h"
+#include "Serialization/JsonSerializer.h"
+#include "UObject/Object.h"
+#include "UObject/ObjectMacros.h"
+
+
+namespace UE::LiveLinkHubLauncherUtils
+{
+	struct FInstalledApp
+	{
+		/** Location of the installed app. */
+		FString InstallLocation;
+
+		/** Namespace of the app. */
+		FString NamespaceId;
+
+		/** Id of the app. */
+		FString ItemId;
+
+		/** Unique ID for the app on the EGS. */
+		FString ArtifactId;
+
+		/** Version of the app. For LiveLinkHub this will correspond to a CL number.  */
+		FString AppVersion;
+
+		/** The apps' internal name. Usually matches the ArtifactId except if the app was using a legacy publishing workflow.  */
+		FString AppName;
+	};
+
+	/** Gather all the installed apps from the Epic Games Launcher. */
+	static bool FindLiveLinkHubInstallation(FInstalledApp& OutLiveLinkHubInfo)
+	{
+		FString InstalledListFile = FString(FPlatformProcess::ApplicationSettingsDir()) / TEXT("UnrealEngineLauncher/LauncherInstalled.dat");
+
+		FString InstalledText;
+		if (FFileHelper::LoadFileToString(InstalledText, *InstalledListFile))
+		{
+			// Deserialize the object
+			TSharedPtr<FJsonObject> RootObject;
+			TSharedRef<TJsonReader<>> Reader = TJsonReaderFactory<>::Create(InstalledText);
+			if (FJsonSerializer::Deserialize(Reader, RootObject) && RootObject.IsValid())
+			{
+				// Parse the list of installations
+				TArray<TSharedPtr<FJsonValue>> InstallationList = RootObject->GetArrayField(TEXT("InstallationList"));
+				for (int32 Idx = 0; Idx < InstallationList.Num(); Idx++)
+				{
+					FInstalledApp InstalledApp;
+
+					TSharedPtr<FJsonObject> InstallationItem = InstallationList[Idx]->AsObject();
+					InstalledApp.AppName = InstallationItem->GetStringField(TEXT("AppName"));
+
+					if (InstalledApp.AppName == GetDefault<ULiveLinkHubEditorSettings>()->LiveLinkHubAppName)
+					{
+						InstalledApp.InstallLocation = InstallationItem->GetStringField(TEXT("InstallLocation"));
+
+						if (!InstalledApp.InstallLocation.Len())
+						{
+							// Shouldn't happen in theory, but just to be safe. 
+							// Doing a continue here instead of returning in case there were somehow multiple LLH installations.
+							continue;
+						}
+
+						const FString& TargetVersion = GetDefault<ULiveLinkHubEditorSettings>()->LiveLinkHubTargetVersion;
+						InstalledApp.AppVersion = InstallationItem->GetStringField(TEXT("AppVersion"));
+
+						if (!TargetVersion.IsEmpty() && InstalledApp.AppVersion != TargetVersion)
+						{
+							// If we target a specific version and it doesn't match the installed app, ignore it.
+							continue;
+						}
+
+						InstalledApp.NamespaceId = InstallationItem->GetStringField(TEXT("NamespaceId"));
+						InstalledApp.ItemId = InstallationItem->GetStringField(TEXT("ItemId"));
+						InstalledApp.ArtifactId = InstallationItem->GetStringField(TEXT("ArtifactId"));
+
+
+						OutLiveLinkHubInfo = MoveTemp(InstalledApp);
+						return true;
+					}
+				}
+			}
+		}
+
+		return false;
+	}
+}
\ No newline at end of file
diff --git a/Engine/Source/Programs/LiveLinkHubEditor/Config/DefaultEditor.ini b/Engine/Source/Programs/LiveLinkHubEditor/Config/DefaultEditor.ini
deleted file mode 100644
