Fix duplicated loctext macro key in LiveLinkHub

diff --git a/Engine/Plugins/Animation/LiveLinkHub/Source/LiveLinkHubEditor/Private/LiveLinkHubEditorModule.cpp b/Engine/Plugins/Animation/LiveLinkHub/Source/LiveLinkHubEditor/Private/LiveLinkHubEditorModule.cpp
--- a/Engine/Plugins/Animation/LiveLinkHub/Source/LiveLinkHubEditor/Private/LiveLinkHubEditorModule.cpp
+++ b/Engine/Plugins/Animation/LiveLinkHub/Source/LiveLinkHubEditor/Private/LiveLinkHubEditorModule.cpp
@@ -88,7 +88,7 @@ void FLiveLinkHubEditorModule::OpenLiveLinkHub()
 			{
 				Notification.SetComplete(
 					LaunchLiveLinkHubErrorTitle,
-					LOCTEXT("LaunchLiveLinkHubError_CouldNotOpenLauncher", "Could not launch LiveLinkHub throug the epic games store."),
+					LOCTEXT("LaunchLiveLinkHubError_CouldNotOpenLauncher", "Could not launch LiveLinkHub through the Epic Games Store."),
 					false
 				);
 
@@ -116,7 +116,7 @@ void FLiveLinkHubEditorModule::OpenLiveLinkHub()
 					{
 						Notification.SetComplete(
 							LaunchLiveLinkHubErrorTitle,
-							LOCTEXT("LaunchLiveLinkHubError_CouldNotOpenLauncher", "Could not find the LiveLink Hub page on the epic games store."),
+							LOCTEXT("LaunchLiveLinkHubError_CouldNotFindHubStorePage", "Could not find the LiveLink Hub page on the Epic Games Store."),
 							false
 						);
 					}
@@ -134,7 +134,7 @@ void FLiveLinkHubEditorModule::OpenLiveLinkHub()
 				{
 					Notification.SetComplete(
 						LaunchLiveLinkHubErrorTitle,
-						LOCTEXT("LaunchLiveLinkHubError_EmptyConfig", "Could not find the LiveLink Hub page on the epic games store, missing configuration for the store page."),
+						LOCTEXT("LaunchLiveLinkHubError_EmptyConfig", "Could not find the LiveLink Hub page on the Epic Games Store, missing configuration for the store page."),
 						false
 					);
 				
