LiveLinkHub - Fix async save crashing on LiveLinkHubEditor target.

diff --git a/Engine/Plugins/Animation/LiveLinkHub/Source/LiveLinkHub/Private/Recording/Implementations/LiveLinkUAssetRecorder.cpp b/Engine/Plugins/Animation/LiveLinkHub/Source/LiveLinkHub/Private/Recording/Implementations/LiveLinkUAssetRecorder.cpp
--- a/Engine/Plugins/Animation/LiveLinkHub/Source/LiveLinkHub/Private/Recording/Implementations/LiveLinkUAssetRecorder.cpp
+++ b/Engine/Plugins/Animation/LiveLinkHub/Source/LiveLinkHub/Private/Recording/Implementations/LiveLinkUAssetRecorder.cpp
@@ -91,11 +91,6 @@ bool FLiveLinkUAssetRecorder::IsRecording() const
 	return bIsRecording;
 }
 
-bool FLiveLinkUAssetRecorder::IsSavingRecording(ULiveLinkRecording* InRecording) const
-{
-	return AsyncSaveTasks.Contains(InRecording);
-}
-
 void FLiveLinkUAssetRecorder::RecordBaseData(FLiveLinkRecordingBaseDataContainer& StaticDataContainer, TSharedPtr<FInstancedStruct>&& DataToRecord)
 {
 	const double TimeNowInSeconds = FPlatformTime::Seconds();
@@ -237,13 +232,14 @@ void FLiveLinkUAssetRecorder::SaveRecording()
 
 		FAssetRegistryModule::AssetCreated(NewRecording);
 
+		const TStrongObjectPtr<ULiveLinkUAssetRecording> NewRecordingStrongPtr(NewRecording);
+		
 		// Create a task to run on a separate thread for saving all frame data and writing the final UAsset to disk.
 		// We use a container rather than just one task on the chance a save operation is still running when another recording is being saved.
-		FAsyncTask<FLiveLinkSaveRecordingAsyncTask>& AsyncTask = *AsyncSaveTasks.Add(NewRecording, MakeUnique<FAsyncTask<FLiveLinkSaveRecordingAsyncTask>>(NewRecording, this));
+		FAsyncTask<FLiveLinkSaveRecordingAsyncTask>& AsyncTask = *AsyncSaveTasks.Add(NewRecordingStrongPtr,
+			MakeUnique<FAsyncTask<FLiveLinkSaveRecordingAsyncTask>>(NewRecording, this));
 		AsyncTask.StartBackgroundTask();
 	}
-
-	return;
 }
 
 void FLiveLinkUAssetRecorder::RecordInitialStaticData()
@@ -261,26 +257,56 @@ void FLiveLinkUAssetRecorder::RecordInitialStaticData()
 	}
 }
 
-void FLiveLinkUAssetRecorder::OnRecordingSaved_GameThread(TWeakObjectPtr<ULiveLinkUAssetRecording> InRecording)
+void FLiveLinkUAssetRecorder::OnRecordingDataSaved_GameThread(FLiveLinkSaveRecordingAsyncTask* InTask)
 {
-	if (InRecording.IsValid())
+	const TStrongObjectPtr<ULiveLinkUAssetRecording> Recording = InTask->GetRecording().Pin();
+	if (ensure(Recording.IsValid()))
 	{
-		const FLiveLinkHubModule& LiveLinkHubModule = FModuleManager::Get().GetModuleChecked<FLiveLinkHubModule>("LiveLinkHub");
-		const TStrongObjectPtr<ULiveLinkRecording> PlaybackRecording = LiveLinkHubModule.GetPlaybackController()->GetRecording();
+		FSavePackageArgs SavePackageArgs;
+		SavePackageArgs.TopLevelFlags = RF_Public | RF_Standalone;
+		SavePackageArgs.Error = GLog;
+		SavePackageArgs.SaveFlags = SAVE_Async;
 
-		UPackage* PackageToUnload = InRecording->GetPackage();
-		const bool bIsPlayingThisRecording = PlaybackRecording.Get() == InRecording.Get();
+		const FString PackageFileName = FPackageName::LongPackageNameToFilename(Recording->GetPackage()->GetName(), FPackageName::GetAssetPackageExtension());
+		if (!UPackage::SavePackage(Recording->GetPackage(), Recording.Get(), *PackageFileName, MoveTemp(SavePackageArgs)))
+		{
+			UE_LOG(LogLiveLinkHub, Error, TEXT("Package '%s' was not saved"), *PackageFileName);
+		}
+	}
+
+	InTask->NotifyPackageSaveStarted();
+}
 
+void FLiveLinkUAssetRecorder::OnRecordingSaveThreadFinished_GameThread(FLiveLinkSaveRecordingAsyncTask* InTask)
+{
+	// Make sure we see the saved file on disk in the asset registry.
+	const FAssetRegistryModule& AssetRegistryModule = FModuleManager::LoadModuleChecked<FAssetRegistryModule>("AssetRegistry");
+	AssetRegistryModule.Get().ScanPathsSynchronous({ TEXT("/Game") }, true);
+	
+	TStrongObjectPtr<ULiveLinkUAssetRecording> Recording = InTask->GetRecording().Pin();
+	if (ensure(Recording.IsValid()))
+	{
 		// Finish task first to make sure strong reference to the recording is cleared.
-		if (const TUniquePtr<FAsyncTask<FLiveLinkSaveRecordingAsyncTask>>* AsyncTask = AsyncSaveTasks.Find(InRecording))
+		if (const TUniquePtr<FAsyncTask<FLiveLinkSaveRecordingAsyncTask>>* AsyncTask = AsyncSaveTasks.Find(Recording))
 		{
 			AsyncTask->Get()->EnsureCompletion();
-			ensure(AsyncSaveTasks.Remove(InRecording) > 0);
+			ensure(AsyncSaveTasks.Remove(Recording) > 0);
 		}
-		
+		else
+		{
+			UE_LOG(LogLiveLinkHub, Error, TEXT("Could not find save task for recording: '%s'"), *Recording->GetName());
+		}
+	
+		const FLiveLinkHubModule& LiveLinkHubModule = FModuleManager::Get().GetModuleChecked<FLiveLinkHubModule>("LiveLinkHub");
+		const TStrongObjectPtr<ULiveLinkRecording> PlaybackRecording = LiveLinkHubModule.GetPlaybackController()->GetRecording();
+
+		UPackage* PackageToUnload = Recording->GetPackage();
+		const bool bIsPlayingThisRecording = PlaybackRecording.Get() == Recording.Get();
+
 		if (!bIsPlayingThisRecording)
 		{
 			// Unload as this is not used again until the user loads it, and allows the bulk animation data to obtain a file handle correctly.
+			Recording.Reset();
 			LiveLinkHubModule.GetPlaybackController()->UnloadRecordingPackage(PackageToUnload);
 		}
 	}
@@ -292,29 +318,31 @@ void FLiveLinkUAssetRecorder::FLiveLinkSaveRecordingAsyncTask::DoWork()
 
 	// Write to bulk data.
 	LiveLinkRecording->SaveRecordingData();
-		
-	FSavePackageArgs SavePackageArgs;
-	SavePackageArgs.TopLevelFlags = RF_Public | RF_Standalone;
-	SavePackageArgs.Error = GLog;
-	SavePackageArgs.SaveFlags = SAVE_Async;
 
-	const FString PackageFileName = FPackageName::LongPackageNameToFilename(LiveLinkRecording->GetPackage()->GetName(), FPackageName::GetAssetPackageExtension());
-	if (UPackage::SavePackage(LiveLinkRecording->GetPackage(), LiveLinkRecording.Get(), *PackageFileName, MoveTemp(SavePackageArgs)))
-	{
-		UPackage::WaitForAsyncFileWrites();
-
-		// Finish on the game thread for safety.
-		TWeakObjectPtr RecordingWeakPtr(LiveLinkRecording.Get());
-		FSimpleDelegateGraphTask::CreateAndDispatchWhenReady(
-			FSimpleDelegateGraphTask::FDelegate::CreateRaw(Recorder, &FLiveLinkUAssetRecorder::OnRecordingSaved_GameThread, RecordingWeakPtr),
-		TStatId(),
-		nullptr,
-		ENamedThreads::GameThread);
-	}
-	else
+	// Let the game thread start the actual package save. Editor targets will assert if SavePackage is called from another thread.
+	FSimpleDelegateGraphTask::CreateAndDispatchWhenReady(
+		FSimpleDelegateGraphTask::FDelegate::CreateRaw(Recorder, &FLiveLinkUAssetRecorder::OnRecordingDataSaved_GameThread, this),
+	TStatId(),
+	nullptr,
+	ENamedThreads::GameThread);
+
+	// Wait for the game thread to signal it has started saving the package.
+	constexpr uint32 TimeoutMillis = 5000;
+	if (!PackageSaveStartedEvent->Wait(TimeoutMillis))
 	{
-		UE_LOG(LogLiveLinkHub, Error, TEXT("Package '%s' was not saved"), *PackageFileName);
+		UE_LOG(LogLiveLinkHub, Error, TEXT("Timed out waiting for package save."));
 	}
+
+	// Block the async thread until the package saves. We wait because the engine doesn't provide any proper callback
+	// when the package finishes saving async, and we still have some cleanup to do once it finishes saving.
+	UPackage::WaitForAsyncFileWrites();
+
+	// Cleanup and finish the thread.
+	FSimpleDelegateGraphTask::CreateAndDispatchWhenReady(
+	FSimpleDelegateGraphTask::FDelegate::CreateRaw(Recorder, &FLiveLinkUAssetRecorder::OnRecordingSaveThreadFinished_GameThread, this),
+	TStatId(),
+	nullptr,
+	ENamedThreads::GameThread);
 }
 
 #undef LOCTEXT_NAMESPACE
diff --git a/Engine/Plugins/Animation/LiveLinkHub/Source/LiveLinkHub/Private/Recording/Implementations/LiveLinkUAssetRecorder.h b/Engine/Plugins/Animation/LiveLinkHub/Source/LiveLinkHub/Private/Recording/Implementations/LiveLinkUAssetRecorder.h
--- a/Engine/Plugins/Animation/LiveLinkHub/Source/LiveLinkHub/Private/Recording/Implementations/LiveLinkUAssetRecorder.h
+++ b/Engine/Plugins/Animation/LiveLinkHub/Source/LiveLinkHub/Private/Recording/Implementations/LiveLinkUAssetRecorder.h
@@ -24,7 +24,6 @@ public:
 	virtual void RecordStaticData(const FLiveLinkSubjectKey& SubjectKey, TSubclassOf<ULiveLinkRole> Role, const FLiveLinkStaticDataStruct& StaticData) override;
 	virtual void RecordFrameData(const FLiveLinkSubjectKey& SubjectKey, const FLiveLinkFrameDataStruct& FrameData) override;
 	virtual bool IsRecording() const override;
-	virtual bool IsSavingRecording(ULiveLinkRecording* InRecording) const override;
 	//~ End ILiveLinkRecorder
 
 private:
@@ -38,16 +37,14 @@ private:
 	void RecordBaseData(FLiveLinkRecordingBaseDataContainer& StaticDataContainer, TSharedPtr<FInstancedStruct>&& DataToRecord);
 	/** Record initial data for all livelink subjects. (Useful when static data was sent before the recording started). */
 	void RecordInitialStaticData();
-	/** Called on the game thread after a recording has been saved. */
-	void OnRecordingSaved_GameThread(TWeakObjectPtr<ULiveLinkUAssetRecording> InRecording);
-
+	
 private:
 	class FLiveLinkSaveRecordingAsyncTask : public FNonAbandonableTask
 	{
 	public:
 		FLiveLinkSaveRecordingAsyncTask(ULiveLinkUAssetRecording* InLiveLinkRecording, FLiveLinkUAssetRecorder* InRecorder)
 		{
-			LiveLinkRecording = TStrongObjectPtr(InLiveLinkRecording);
+			LiveLinkRecording = InLiveLinkRecording;
 			Recorder = InRecorder;
 		}
 
@@ -57,16 +54,28 @@ private:
 		}
 
 		void DoWork();
+		
+		/** Notify we have started saving the package. */
+		void NotifyPackageSaveStarted() { PackageSaveStartedEvent->Trigger(); }
+		/** The recording being saved by this task. */
+		TWeakObjectPtr<ULiveLinkUAssetRecording> GetRecording() const { return LiveLinkRecording; }
 
 	private:
 		/** The recording being saved. */
-		TStrongObjectPtr<ULiveLinkUAssetRecording> LiveLinkRecording;
+		TWeakObjectPtr<ULiveLinkUAssetRecording> LiveLinkRecording;
 		/** The recorder owner. */
 		FLiveLinkUAssetRecorder* Recorder = nullptr;
+		/** If the game thread has started to save the package. */
+		FEventRef PackageSaveStartedEvent;
 	};
 
+	/** Called on the game thread after the recording data has been saved. */
+	void OnRecordingDataSaved_GameThread(FLiveLinkSaveRecordingAsyncTask* InTask);
+	/** Called when the async save thread has finished. */
+	void OnRecordingSaveThreadFinished_GameThread(FLiveLinkSaveRecordingAsyncTask* InTask);
+	
 	/** Current async save tasks. */
-	TMap<TWeakObjectPtr<ULiveLinkRecording>, TUniquePtr<FAsyncTask<FLiveLinkSaveRecordingAsyncTask>>> AsyncSaveTasks;
+	TMap<TStrongObjectPtr<ULiveLinkUAssetRecording>, TUniquePtr<FAsyncTask<FLiveLinkSaveRecordingAsyncTask>>> AsyncSaveTasks;
 	/** Holds metadata and recording data. */
 	TPimplPtr<FLiveLinkUAssetRecordingData> CurrentRecording;
 	/** Whether we're currently recording livelink data. */
diff --git a/Engine/Plugins/Animation/LiveLinkHub/Source/LiveLinkHub/Private/Recording/LiveLinkHubRecordingController.cpp b/Engine/Plugins/Animation/LiveLinkHub/Source/LiveLinkHub/Private/Recording/LiveLinkHubRecordingController.cpp
--- a/Engine/Plugins/Animation/LiveLinkHub/Source/LiveLinkHub/Private/Recording/LiveLinkHubRecordingController.cpp
+++ b/Engine/Plugins/Animation/LiveLinkHub/Source/LiveLinkHub/Private/Recording/LiveLinkHubRecordingController.cpp
@@ -37,11 +37,6 @@ bool FLiveLinkHubRecordingController::IsRecording() const
 	return RecorderImplementation->IsRecording();
 }
 
-bool FLiveLinkHubRecordingController::IsSavingRecording(ULiveLinkRecording* InRecording) const
-{
-	return RecorderImplementation->IsSavingRecording(InRecording);
-}
-
 void FLiveLinkHubRecordingController::RecordStaticData(const FLiveLinkSubjectKey& SubjectKey, TSubclassOf<ULiveLinkRole> Role, const FLiveLinkStaticDataStruct& StaticData)
 {
 	RecorderImplementation->RecordStaticData(SubjectKey, Role, StaticData);
diff --git a/Engine/Plugins/Animation/LiveLinkHub/Source/LiveLinkHub/Private/Recording/LiveLinkHubRecordingController.h b/Engine/Plugins/Animation/LiveLinkHub/Source/LiveLinkHub/Private/Recording/LiveLinkHubRecordingController.h
--- a/Engine/Plugins/Animation/LiveLinkHub/Source/LiveLinkHub/Private/Recording/LiveLinkHubRecordingController.h
+++ b/Engine/Plugins/Animation/LiveLinkHub/Source/LiveLinkHub/Private/Recording/LiveLinkHubRecordingController.h
@@ -27,9 +27,6 @@ public:
 	/** Returns whether we're currently recording. */
 	bool IsRecording() const;
 
-	/** Returns whether the recording is currently saving. */
-	bool IsSavingRecording(ULiveLinkRecording* InRecording) const;
-
 	/** Record static data in the current recording. */
 	void RecordStaticData(const FLiveLinkSubjectKey& SubjectKey, TSubclassOf<class ULiveLinkRole> Role, const FLiveLinkStaticDataStruct& StaticData);
 
diff --git a/Engine/Plugins/Animation/LiveLinkHub/Source/LiveLinkHub/Private/Recording/LiveLinkRecorder.h b/Engine/Plugins/Animation/LiveLinkHub/Source/LiveLinkHub/Private/Recording/LiveLinkRecorder.h
--- a/Engine/Plugins/Animation/LiveLinkHub/Source/LiveLinkHub/Private/Recording/LiveLinkRecorder.h
+++ b/Engine/Plugins/Animation/LiveLinkHub/Source/LiveLinkHub/Private/Recording/LiveLinkRecorder.h
@@ -23,6 +23,4 @@ public:
 	virtual void RecordFrameData(const FLiveLinkSubjectKey& SubjectKey, const FLiveLinkFrameDataStruct& FrameData) = 0;
 	/** Return whether we are currently recording livelink data. */
 	virtual bool IsRecording() const = 0;
-	/** Return whether the recording is in the process of being saved. */
-	virtual bool IsSavingRecording(ULiveLinkRecording* InRecording) const = 0;
 };
