LiveLink Hub: Build fix; follow-up to CL 35779711.

diff --git a/Engine/Plugins/Animation/LiveLinkHub/Source/LiveLinkHub/Private/UI/Widgets/SLiveLinkHubTabViewWithManagerBase.cpp b/Engine/Plugins/Animation/LiveLinkHub/Source/LiveLinkHub/Private/UI/Widgets/SLiveLinkHubTabViewWithManagerBase.cpp
--- a/Engine/Plugins/Animation/LiveLinkHub/Source/LiveLinkHub/Private/UI/Widgets/SLiveLinkHubTabViewWithManagerBase.cpp
+++ b/Engine/Plugins/Animation/LiveLinkHub/Source/LiveLinkHub/Private/UI/Widgets/SLiveLinkHubTabViewWithManagerBase.cpp
@@ -6,7 +6,6 @@
 #include "Framework/Commands/UICommandList.h"
 #include "Framework/Docking/LayoutService.h"
 #include "Framework/MultiBox/MultiBoxBuilder.h"
-#include "ISessionFrontendModule.h"
 #include "ISlateReflectorModule.h"
 #include "LiveLinkHub.h"
 #include "LiveLinkHubCommands.h"
@@ -16,6 +15,11 @@
 #include "Modules/ModuleManager.h"
 #include "UI/Widgets/SLiveLinkHubTabViewBase.h"
 
+#if WITH_UNREAL_TARGET_DEVELOPER_TOOLS
+#	include "ISessionFrontendModule.h"
+#endif
+
+
 #define LOCTEXT_NAMESPACE "LiveLinkHub.SLiveLinkHubTabViewWithManagerBase"
 
 void SLiveLinkHubTabViewWithManagerBase::Construct(const FArguments& InArgs)
@@ -102,6 +106,7 @@ void SLiveLinkHubTabViewWithManagerBase::FillFileMenu(FMenuBuilder& MenuBuilder)
 	if (bShowDevTools)
 	{
 		MenuBuilder.BeginSection("Development", LOCTEXT("DevelopmentHeader", "Development"));
+#if WITH_UNREAL_TARGET_DEVELOPER_TOOLS
 		MenuBuilder.AddMenuEntry(
 			LOCTEXT("FileMenu_Development_AutomationTools", "Automation Tools"),
 			FText(),
@@ -113,6 +118,7 @@ void SLiveLinkHubTabViewWithManagerBase::FillFileMenu(FMenuBuilder& MenuBuilder)
 				SessionFrontend.InvokeSessionFrontend(FName("AutomationPanel"));
 			})
 		);
+#endif // #if WITH_UNREAL_TARGET_DEVELOPER_TOOLS
 		MenuBuilder.AddMenuEntry(
 			LOCTEXT("FileMenu_Development_WidgetReflector", "Widget Reflector"),
 			FText(),
