LiveLink Hub: Another build fix; follow-up to CL 35779711.

diff --git a/Engine/Plugins/Animation/LiveLinkHub/Source/LiveLinkHub/Private/UI/Widgets/SLiveLinkHubTabViewWithManagerBase.cpp b/Engine/Plugins/Animation/LiveLinkHub/Source/LiveLinkHub/Private/UI/Widgets/SLiveLinkHubTabViewWithManagerBase.cpp
--- a/Engine/Plugins/Animation/LiveLinkHub/Source/LiveLinkHub/Private/UI/Widgets/SLiveLinkHubTabViewWithManagerBase.cpp
+++ b/Engine/Plugins/Animation/LiveLinkHub/Source/LiveLinkHub/Private/UI/Widgets/SLiveLinkHubTabViewWithManagerBase.cpp
@@ -6,7 +6,6 @@
 #include "Framework/Commands/UICommandList.h"
 #include "Framework/Docking/LayoutService.h"
 #include "Framework/MultiBox/MultiBoxBuilder.h"
-#include "ISlateReflectorModule.h"
 #include "LiveLinkHub.h"
 #include "LiveLinkHubCommands.h"
 #include "LiveLinkHubModule.h"
@@ -15,8 +14,11 @@
 #include "Modules/ModuleManager.h"
 #include "UI/Widgets/SLiveLinkHubTabViewBase.h"
 
-#if WITH_UNREAL_TARGET_DEVELOPER_TOOLS
-#	include "ISessionFrontendModule.h"
+#if !UE_BUILD_SHIPPING
+#	if WITH_UNREAL_TARGET_DEVELOPER_TOOLS
+#		include "ISessionFrontendModule.h"
+#	endif
+#	include "ISlateReflectorModule.h"
 #endif
 
 
