LiveLinkHub - Add "Open File Location" to the recording list context menu.

diff --git a/Engine/Plugins/Animation/LiveLinkHub/Source/LiveLinkHub/Private/Recording/SLiveLinkHubRecordingListView.h b/Engine/Plugins/Animation/LiveLinkHub/Source/LiveLinkHub/Private/Recording/SLiveLinkHubRecordingListView.h
--- a/Engine/Plugins/Animation/LiveLinkHub/Source/LiveLinkHub/Private/Recording/SLiveLinkHubRecordingListView.h
+++ b/Engine/Plugins/Animation/LiveLinkHub/Source/LiveLinkHub/Private/Recording/SLiveLinkHubRecordingListView.h
@@ -9,9 +9,12 @@
 #include "Delegates/DelegateCombinations.h"
 #include "Framework/Commands/GenericCommands.h"
 #include "Framework/MultiBox/MultiBoxBuilder.h"
+#include "HAL/PlatformProcess.h"
 #include "IContentBrowserSingleton.h"
 #include "Implementations/LiveLinkUAssetRecording.h"
 #include "LiveLinkHubModule.h"
+#include "Misc/PackageName.h"
+#include "Misc/Paths.h"
 #include "Recording/LiveLinkHubPlaybackController.h"
 #include "Recording/LiveLinkRecording.h"
 #include "UObject/SavePackage.h"
@@ -313,6 +316,27 @@ private:
 					FCanExecuteAction::CreateLambda([] () { return true; })
 				)
 			);
+
+			MenuBuilder.AddMenuEntry(
+				LOCTEXT("OpenFileLocationLabel", "Open File Location..."),
+				LOCTEXT("OpenFileLocationTooltip", "Open the folder containing this file"),
+				FSlateIcon(FAppStyle::GetAppStyleSetName(), "Icons.FolderOpen"),
+				FUIAction(
+					FExecuteAction::CreateLambda([SelectedAsset] ()
+					{
+						if (SelectedAsset.IsValid())
+						{
+							const FString PackageName = SelectedAsset->GetPathName();
+							const FString AssetFilePath = FPackageName::LongPackageNameToFilename(PackageName, FPackageName::GetAssetPackageExtension());
+							const FString AbsoluteFilePath = FPaths::ConvertRelativePathToFull(AssetFilePath);
+							const FString AssetDirectory = FPaths::GetPath(AbsoluteFilePath);
+						
+							FPlatformProcess::ExploreFolder(*AssetDirectory);
+						}
+					}),
+					FCanExecuteAction::CreateLambda([SelectedAsset] () { return SelectedAsset.IsValid(); })
+				)
+			);
 		}
 		MenuBuilder.EndSection();
 
