LiveLinkHub - Improve frame buffering and fix framerate calculations, which were impacting recording length and playback speed.

- Framerate as viewed in LLH is now based on the total frames by the final timestamp. The correct framerate is still determined by the client.
- Length and selections are now properly calculated.
- Frame buffering is now more intelligent and efficient, and buffer in ranges to help account for scrubbing.
- Optimize bulk data reads, by reading more frames at once from disk in batches before processing them.
- Fix an issue where scrubbing to a position not buffered would result in the wrong frames being played back.
- Fix race condition when buffering multiple sources. We now alternate between sources, loading in the correct range before reporting the frames as ready.
- Fix possible race condition crash where frame memory was freed early.
- Fix thread safe access of recording playback delegate.
- Add stats profiling.

diff --git a/Engine/Plugins/Animation/LiveLinkHub/Source/LiveLinkHub/Private/Data/LiveLinkHubBulkData.cpp b/Engine/Plugins/Animation/LiveLinkHub/Source/LiveLinkHub/Private/Data/LiveLinkHubBulkData.cpp
--- a/Engine/Plugins/Animation/LiveLinkHub/Source/LiveLinkHub/Private/Data/LiveLinkHubBulkData.cpp
+++ b/Engine/Plugins/Animation/LiveLinkHub/Source/LiveLinkHub/Private/Data/LiveLinkHubBulkData.cpp
@@ -2,6 +2,8 @@
 
 #include "LiveLinkHubBulkData.h"
 
+#include "LiveLinkHubLog.h"
+
 #include "Serialization/BufferArchive.h"
 #include "Serialization/MemoryReader.h"
 
@@ -34,10 +36,12 @@ void FLiveLinkHubBulkData::ReadBulkData(const int64 InBytesToRead, uint8* InMemo
 	BulkDataOffset = ReadBulkDataImpl(BulkDataOffset, InBytesToRead, InMemory);
 }
 
-FLiveLinkHubBulkData::FScopedBulkDataMemoryReader FLiveLinkHubBulkData::CreateBulkDataMemoryReader(const int64 InBytesToRead)
+TSharedPtr<FLiveLinkHubBulkData::FScopedBulkDataMemoryReader> FLiveLinkHubBulkData::CreateBulkDataMemoryReader(const int64 InBytesToRead)
 {
-	FScopedBulkDataMemoryReader Reader(BulkDataOffset, InBytesToRead, this);
-	BulkDataOffset = Reader.GetBulkDataOffset();
+	DECLARE_SCOPE_CYCLE_COUNTER(TEXT("FLiveLinkHubBulkData::CreateBulkDataMemoryReader"), STAT_FLiveLinkHubBulkData_CreateBulkDataMemoryReader, STATGROUP_LiveLinkHub);
+	
+	TSharedPtr<FScopedBulkDataMemoryReader> Reader = MakeShared<FScopedBulkDataMemoryReader>(BulkDataOffset, InBytesToRead, this);
+	BulkDataOffset = Reader->GetBulkDataOffset();
 
 	return Reader;
 }
@@ -69,6 +73,8 @@ void FLiveLinkHubBulkData::WriteBulkData(TArray64<uint8>& Data)
 
 int64 FLiveLinkHubBulkData::ReadBulkDataImpl(int64 InOffset, int64 InBytesToRead, uint8* InMemory)
 {
+	DECLARE_SCOPE_CYCLE_COUNTER(TEXT("FLiveLinkHubBulkData::ReadBulkDataImpl"), STAT_FLiveLinkHubBulkData_ReadBulkDataImpl, STATGROUP_LiveLinkHub);
+	
 	if (!RecordingFileReader.IsValid())
 	{
 		check(BulkData.DoesExist());
diff --git a/Engine/Plugins/Animation/LiveLinkHub/Source/LiveLinkHub/Private/Data/LiveLinkHubBulkData.h b/Engine/Plugins/Animation/LiveLinkHub/Source/LiveLinkHub/Private/Data/LiveLinkHubBulkData.h
--- a/Engine/Plugins/Animation/LiveLinkHub/Source/LiveLinkHub/Private/Data/LiveLinkHubBulkData.h
+++ b/Engine/Plugins/Animation/LiveLinkHub/Source/LiveLinkHub/Private/Data/LiveLinkHubBulkData.h
@@ -55,7 +55,7 @@ public:
 	void ReadBulkData(const int64 InBytesToRead, uint8* InMemory);
 
 	/** Create a scoped memory reader consisting of the bulk data bytes read. Increments file offset. */
-	FScopedBulkDataMemoryReader CreateBulkDataMemoryReader(const int64 InBytesToRead);
+	TSharedPtr<FScopedBulkDataMemoryReader> CreateBulkDataMemoryReader(const int64 InBytesToRead);
 	
 	/** Reset to the initial offset in the file. */
 	void ResetBulkDataOffset();
diff --git a/Engine/Plugins/Animation/LiveLinkHub/Source/LiveLinkHub/Private/LiveLinkHubClient.cpp b/Engine/Plugins/Animation/LiveLinkHub/Source/LiveLinkHub/Private/LiveLinkHubClient.cpp
--- a/Engine/Plugins/Animation/LiveLinkHub/Source/LiveLinkHub/Private/LiveLinkHubClient.cpp
+++ b/Engine/Plugins/Animation/LiveLinkHub/Source/LiveLinkHub/Private/LiveLinkHubClient.cpp
@@ -18,7 +18,6 @@
 #include "UObject/StrongObjectPtr.h"
 #include "UObject/UObjectGlobals.h"
 
-DECLARE_STATS_GROUP(TEXT("Live Link Hub"), STATGROUP_LiveLinkHub, STATCAT_Advanced);
 DECLARE_CYCLE_STAT(TEXT("LiveLinkHub - Push StaticData"), STAT_LiveLinkHub_PushStaticData, STATGROUP_LiveLinkHub);
 DECLARE_CYCLE_STAT(TEXT("LiveLinkHub - Push FrameData"), STAT_LiveLinkHub_PushFrameData, STATGROUP_LiveLinkHub);
 
diff --git a/Engine/Plugins/Animation/LiveLinkHub/Source/LiveLinkHub/Private/LiveLinkHubLog.h b/Engine/Plugins/Animation/LiveLinkHub/Source/LiveLinkHub/Private/LiveLinkHubLog.h
--- a/Engine/Plugins/Animation/LiveLinkHub/Source/LiveLinkHub/Private/LiveLinkHubLog.h
+++ b/Engine/Plugins/Animation/LiveLinkHub/Source/LiveLinkHub/Private/LiveLinkHubLog.h
@@ -3,5 +3,8 @@
 #pragma once
 
 #include "Logging/LogMacros.h"
+#include "Stats/Stats2.h"
 
 DECLARE_LOG_CATEGORY_EXTERN(LogLiveLinkHub, Log, All);
+
+DECLARE_STATS_GROUP(TEXT("Live Link Hub"), STATGROUP_LiveLinkHub, STATCAT_Advanced);
\ No newline at end of file
diff --git a/Engine/Plugins/Animation/LiveLinkHub/Source/LiveLinkHub/Private/Recording/Implementations/LiveLinkFrameData.cpp b/Engine/Plugins/Animation/LiveLinkHub/Source/LiveLinkHub/Private/Recording/Implementations/LiveLinkFrameData.cpp
new file mode 100644
--- /dev/null
+++ b/Engine/Plugins/Animation/LiveLinkHub/Source/LiveLinkHub/Private/Recording/Implementations/LiveLinkFrameData.cpp
@@ -0,0 +1,86 @@
+// Copyright Epic Games, Inc. All Rights Reserved.
+
+#include "LiveLinkFrameData.h"
+
+#include "Settings/LiveLinkHubSettings.h"
+
+namespace UE::LiveLinkHub::FrameData::Private
+{
+	void FFrameBufferCache::TrimCache()
+	{
+		const int32 MaxHistory = GetDefault<ULiveLinkHubSettings>()->MaxBufferRangeHistory;
+		const int32 AmountToTrim = FrameData.Num() - MaxHistory;
+		if (AmountToTrim > 0)
+		{
+			FrameData.RemoveAt(0, AmountToTrim);
+		}
+	}
+
+	void FFrameBufferCache::CleanCache(const TRange<int32>& InActiveRange)
+	{
+		for (FLiveLinkRecordingBaseDataContainer& Container : FrameData)
+		{
+			const TRange<int32> BufferedRange = Container.GetBufferedFrames();
+
+			// Clear portions that are completely outside the active range.
+			if (BufferedRange.GetUpperBoundValue() < InActiveRange.GetLowerBoundValue()
+				|| BufferedRange.GetLowerBoundValue() > InActiveRange.GetUpperBoundValue())
+			{
+				Container.ClearData();
+			}
+			else
+			{
+				// In this case we remove frames beyond the intersection with the active range,
+				// while keeping frames that are still relevant.
+				TRange<int32> Intersection = TRange<int32>::Intersection(InActiveRange, BufferedRange);
+				Container.RemoveFramesAfter(Intersection.GetUpperBoundValue());
+				Container.RemoveFramesBefore(Intersection.GetLowerBoundValue());
+			}
+		}
+		
+		// Cache data is also cleared at a per frame level when checking history.
+		FrameData.RemoveAll([](const FLiveLinkRecordingBaseDataContainer& Container)
+		{
+			return Container.IsEmpty();
+		});
+
+		TrimCache();
+	}
+
+	RangeHelpers::Private::TRangeArray<int32> FFrameBufferCache::GetCacheBufferRanges() const
+	{
+		RangeHelpers::Private::TRangeArray<int32> Results;
+		for (const FLiveLinkRecordingBaseDataContainer& Container : FrameData)
+		{
+			Results.Add(Container.GetBufferedFrames());
+		}
+
+		return Results;
+	}
+
+	TSharedPtr<FInstancedStruct> FFrameBufferCache::TryGetCachedFrame(const int32 InFrame, double& OutTimestamp)
+	{
+		for (FLiveLinkRecordingBaseDataContainer& Data : FrameData)
+		{
+			if (TSharedPtr<FInstancedStruct> Result = Data.TryGetFrame(InFrame, OutTimestamp))
+			{
+				return Result;
+			}
+		}
+
+		return nullptr;
+	}
+
+	bool FFrameBufferCache::ContainsFrame(const int32 InFrame) const
+	{
+		for (const FLiveLinkRecordingBaseDataContainer& Data : FrameData)
+		{
+			if (Data.IsFrameLoaded(InFrame))
+			{
+				return true;
+			}
+		}
+
+		return false;
+	}
+}
diff --git a/Engine/Plugins/Animation/LiveLinkHub/Source/LiveLinkHub/Private/Recording/Implementations/LiveLinkFrameData.h b/Engine/Plugins/Animation/LiveLinkHub/Source/LiveLinkHub/Private/Recording/Implementations/LiveLinkFrameData.h
new file mode 100644
--- /dev/null
+++ b/Engine/Plugins/Animation/LiveLinkHub/Source/LiveLinkHub/Private/Recording/Implementations/LiveLinkFrameData.h
@@ -0,0 +1,129 @@
+// Copyright Epic Games, Inc. All Rights Reserved.
+
+#pragma once
+
+#include "LiveLinkRecordingDataContainer.h"
+#include "Recording/LiveLinkRecordingRangeHelpers.h"
+#include "StructUtils/InstancedStruct.h"
+
+namespace UE::LiveLinkHub::FrameData::Private
+{
+	/** Data manager for frames which have been loaded. */
+	struct FFrameBufferCache
+	{
+		/** Loaded frame data. */
+		TArray<FLiveLinkRecordingBaseDataContainer> FrameData;
+
+		/** Ensure the size of the cache is limited, removing older entries. */
+		void TrimCache();
+
+		/** Remove data from the cache that is no longer needed. */
+		void CleanCache(const TRange<int32>& InActiveRange);
+
+		/** Return all ranges contained within the cache. */
+		RangeHelpers::Private::TRangeArray<int32> GetCacheBufferRanges() const;
+
+		/** Look through the cache for a loaded frame. */
+		TSharedPtr<FInstancedStruct> TryGetCachedFrame(const int32 InFrame, double& OutTimestamp);
+
+		/** Checks if a frame is contained within the cache. */
+		bool ContainsFrame(const int32 InFrame) const;
+	};
+
+	/** Keeps track of current frame iterations while buffering. */
+	struct FFrameBufferIterationData
+	{
+		enum ECompletionStatus
+		{
+			// New iteration, no processing has performed yet.
+			New,
+			// Iteration is currently active.
+			Active,
+			// Iteration has been canceled for any reason.
+			Canceled,
+			// Iteration has fully completed, buffering all available requested frames.
+			Complete
+		};
+
+		/** Temporary storage for loaded frames, which gets moved to data containers after loading. */
+		struct FTemporaryData
+		{
+			/** Fully processed frame data. */
+			TArray<TSharedPtr<FInstancedStruct>> RecordedData;
+			/** Timestamps corresponding to the frame data. */
+			TArray<double> Timestamps;
+			void Reset()
+			{
+				RecordedData.Empty();
+				Timestamps.Empty();
+			}
+		};
+
+		/** The last loaded left frame. */
+		int32 LastLoadedLeftFrame;
+		/** The last loaded right frame. */
+		int32 LastLoadedRightFrame;
+		/** If we are loading to the right. */
+		bool bLoadRight;
+		/** Current status of this iteration. */
+		ECompletionStatus Status;
+
+		/** Forward-looking frame data. */
+		FTemporaryData ForwardData;
+		/** Reverse looking frame data. */
+		FTemporaryData ReverseData;
+
+		FFrameBufferIterationData()
+		{
+			Reset();
+		}
+
+		/** Resets the iteration to a new iteration, clearing all storage. */
+		void Reset()
+		{
+			Status = New;
+			LastLoadedLeftFrame = INDEX_NONE;
+			LastLoadedRightFrame = INDEX_NONE;
+			bLoadRight = true;
+			ForwardData.Reset();
+			ReverseData.Reset();
+		}
+	};
+	
+	/** Frame data information when loading from a recording file. */
+	struct FFrameMetaData
+	{
+		/** The subject key used for the frame data. */
+		TSharedPtr<FLiveLinkSubjectKey> FrameDataSubjectKey;
+		/** The struct for this frame data. */
+		TWeakObjectPtr<UScriptStruct> LoadedStruct;
+		/** The position in the file recording where frame data begins. */
+		int64 RecordingStartFrameFilePosition = 0;
+		/** Maximum number of frames. */
+		int32 MaxFrames = 0;
+		/** Total size of the structure. */
+		int32 SerializedStructureSize = 0;
+		/** The size in bytes of each animation frame. */
+		int32 FrameDiskSize;
+		/** The last timestamp for this frame data. */
+		double LastTimestamp = 0.f;
+		/** The frame rate, based only off of number of frames and the last timestamp. */
+		FFrameRate LocalFrameRate;
+		/** Cache of previously buffered frames, which are not currently active. */
+		FFrameBufferCache BufferedCache;
+		/** Current iteration data while buffering. */
+		FFrameBufferIterationData BufferIterationData;
+	
+		/** Find the correct file offset based on the frame index. */
+		int64 GetFrameFilePosition(const int32 InFrameIdx) const
+		{
+			return RecordingStartFrameFilePosition + GetRelativeFrameFilePosition(InFrameIdx);
+		};
+
+		/** Find the offset relative to local storage only, not accounting for disk position. */
+		int64 GetRelativeFrameFilePosition(const int32 InFrameIdx) const
+		{
+			return FrameDiskSize * InFrameIdx;
+		};
+	};
+}
\ No newline at end of file
diff --git a/Engine/Plugins/Animation/LiveLinkHub/Source/LiveLinkHub/Private/Recording/Implementations/LiveLinkRecordingDataContainer.cpp b/Engine/Plugins/Animation/LiveLinkHub/Source/LiveLinkHub/Private/Recording/Implementations/LiveLinkRecordingDataContainer.cpp
new file mode 100644
--- /dev/null
+++ b/Engine/Plugins/Animation/LiveLinkHub/Source/LiveLinkHub/Private/Recording/Implementations/LiveLinkRecordingDataContainer.cpp
@@ -0,0 +1,84 @@
+// Copyright Epic Games, Inc. All Rights Reserved.
+
+#include "LiveLinkRecordingDataContainer.h"
+
+#include "Recording/LiveLinkRecordingRangeHelpers.h"
+
+bool FLiveLinkRecordingBaseDataContainer::IsEmpty() const
+{
+	check(Timestamps.Num() == RecordedData.Num());
+	return Timestamps.Num() == 0;
+}
+
+void FLiveLinkRecordingBaseDataContainer::ClearData()
+{
+	Timestamps.Empty();
+	RecordedData.Empty();
+	RecordedDataStartFrame = 0;
+}
+
+TRange<int32> FLiveLinkRecordingBaseDataContainer::GetBufferedFrames() const
+{
+	if (Timestamps.Num() == 0)
+	{
+		return TRange<int32>::Empty();
+	}
+	const int32 EndFrame = FMath::Max(0, RecordedDataStartFrame + Timestamps.Num() - 1);
+	return UE::LiveLinkHub::RangeHelpers::Private::MakeInclusiveRange(RecordedDataStartFrame, EndFrame);
+}
+
+TSharedPtr<FInstancedStruct> FLiveLinkRecordingBaseDataContainer::TryGetFrame(const int32 InFrame) const
+{
+	if (IsFrameLoaded(InFrame))
+	{
+		const int32 RelativeFrameIdx = GetRelativeFrameIndex(InFrame);
+		return RecordedData[RelativeFrameIdx];
+	}
+	return nullptr;
+}
+
+TSharedPtr<FInstancedStruct> FLiveLinkRecordingBaseDataContainer::TryGetFrame(const int32 InFrame, double& OutTimestamp) const
+{
+	if (TSharedPtr<FInstancedStruct> OutFrame = TryGetFrame(InFrame))
+	{
+		const int32 RelativeFrameIdx = GetRelativeFrameIndex(InFrame);
+		OutTimestamp = Timestamps[RelativeFrameIdx];
+		return OutFrame;
+	}
+	return nullptr;
+}
+
+void FLiveLinkRecordingBaseDataContainer::RemoveFramesBefore(const int32 InEndFrame)
+{
+	const int32 RelativeFrameIdx = FMath::Min(GetRelativeFrameIndex(InEndFrame), RecordedData.Num() - 1);
+	if (!IsEmpty() && RelativeFrameIdx >= 0)
+	{
+		const int32 AmountToRemove = RelativeFrameIdx + 1;
+		
+		RecordedData.RemoveAt(0, AmountToRemove);
+		Timestamps.RemoveAt(0, AmountToRemove);
+
+		RecordedDataStartFrame += AmountToRemove;
+	}
+}
+
+void FLiveLinkRecordingBaseDataContainer::RemoveFramesAfter(const int32 InStartFrame)
+{
+	const int32 RelativeFrameIdx = GetRelativeFrameIndex(InStartFrame);
+	if (!IsEmpty() && RelativeFrameIdx >= 0 && RelativeFrameIdx < RecordedData.Num())
+	{
+		const int32 AmountToRemove = RecordedData.Num() - RelativeFrameIdx;
+		
+		RecordedData.RemoveAt(RelativeFrameIdx, AmountToRemove);
+		Timestamps.RemoveAt(RelativeFrameIdx, AmountToRemove);
+	}
+}
+
+void FLiveLinkRecordingBaseDataContainer::ValidateData() const
+{
+	check(Timestamps.Num() == RecordedData.Num());
+	for (const TSharedPtr<FInstancedStruct>& InstancedStruct : RecordedData)
+	{
+		check(InstancedStruct.IsValid() && InstancedStruct->IsValid());
+	}
+}
diff --git a/Engine/Plugins/Animation/LiveLinkHub/Source/LiveLinkHub/Private/Recording/Implementations/LiveLinkRecordingDataContainer.h b/Engine/Plugins/Animation/LiveLinkHub/Source/LiveLinkHub/Private/Recording/Implementations/LiveLinkRecordingDataContainer.h
new file mode 100644
--- /dev/null
+++ b/Engine/Plugins/Animation/LiveLinkHub/Source/LiveLinkHub/Private/Recording/Implementations/LiveLinkRecordingDataContainer.h
@@ -0,0 +1,119 @@
+// Copyright Epic Games, Inc. All Rights Reserved.
+
+#pragma once
+
+#include "LiveLinkRole.h"
+#include "LiveLinkTypes.h"
+#include "StructUtils/InstancedStruct.h"
+
+#include "LiveLinkRecordingDataContainer.generated.h"
+
+/** Base data container for a recording track. */
+USTRUCT()
+struct FLiveLinkRecordingBaseDataContainer
+{
+	GENERATED_BODY()
+
+	/** SERIALIZED DATA - Timestamps for the recorded data. Each entry matches an entry in the RecordedData array. */
+	TArray<double> Timestamps;
+	
+	/**
+	 * SERIALIZED DATA - Array of either static or frame recorded for a given timestamp.
+	 * TSharedPtr used as streaming the data in may require shared access.
+	 * FSharedStruct is not used because it doesn't implement Serialize().
+	 */
+	TArray<TSharedPtr<FInstancedStruct>> RecordedData;
+
+	/** The current start frame for RecordedData. */
+	int32 RecordedDataStartFrame = 0;
+
+	/** If this container contains no data. */
+	bool IsEmpty() const;
+
+	/** Empty all data. */
+	void ClearData();
+
+	/** Retrieve the total buffered frames. */
+	TRange<int32> GetBufferedFrames() const;
+
+	/**
+	 * Retrieve a loaded frame.
+	 *
+	 * @param InFrame The absolute frame index to load.
+	 *
+	 * @return The frame if it exists.
+	 */
+	TSharedPtr<FInstancedStruct> TryGetFrame(const int32 InFrame) const;
+
+	/**
+	 * Retrieve a loaded frame.
+	 *
+	 * @param InFrame The absolute frame index to load.
+	 * @param OutTimestamp Return the timestamp, if it exists.
+	 *
+	 * @return The frame if it exists.
+	 */
+	TSharedPtr<FInstancedStruct> TryGetFrame(const int32 InFrame, double& OutTimestamp) const;
+
+	/**
+	 * Remove all frames before, and including, the input frame.
+	 * @param InEndFrame The final frame to remove, inclusive.
+	 */
+	void RemoveFramesBefore(const int32 InEndFrame);
+
+	/**
+	 * Remove all frames including and after the input frame.
+	 * @param InStartFrame The first frame to remove, inclusive.
+	 */
+	void RemoveFramesAfter(const int32 InStartFrame);
+
+	/** Convert the absolute frame to the relative index for this data container. */
+	int32 GetRelativeFrameIndex(const int32 InFrame) const
+	{
+		return InFrame - RecordedDataStartFrame;
+	}
+	
+	/**
+	 * Checks if a frame is currently loaded.
+	 * 
+	 * @param InFrame The absolute frame index to check.
+	 *
+	 * @return True if the frame index exists within the frame array.
+	 */
+	bool IsFrameLoaded(const int32 InFrame) const
+	{
+		return InFrame >= RecordedDataStartFrame && InFrame < RecordedDataStartFrame + RecordedData.Num();
+	}
+
+	/** Check data memory is valid and expected. */
+	void ValidateData() const;
+};
+
+/** Container for static data. */
+USTRUCT()
+struct FLiveLinkRecordingStaticDataContainer : public FLiveLinkRecordingBaseDataContainer
+{
+	GENERATED_BODY()
+
+	/** The role of the static data being recorded. */
+	UPROPERTY()
+	TSubclassOf<ULiveLinkRole> Role = nullptr;
+};
+
+USTRUCT()
+struct FLiveLinkUAssetRecordingData
+{
+	GENERATED_BODY()
+
+	/** Length of the recording in seconds. */
+	UPROPERTY()
+	double LengthInSeconds = 0;
+
+	/** Static data encountered while recording. */
+	UPROPERTY()
+	TMap<FLiveLinkSubjectKey, FLiveLinkRecordingStaticDataContainer> StaticData;
+	
+	/** Frame data encountered while recording. */
+	UPROPERTY()
+	TMap<FLiveLinkSubjectKey, FLiveLinkRecordingBaseDataContainer> FrameData;
+};
diff --git a/Engine/Plugins/Animation/LiveLinkHub/Source/LiveLinkHub/Private/Recording/Implementations/LiveLinkUAssetRecording.cpp b/Engine/Plugins/Animation/LiveLinkHub/Source/LiveLinkHub/Private/Recording/Implementations/LiveLinkUAssetRecording.cpp
--- a/Engine/Plugins/Animation/LiveLinkHub/Source/LiveLinkHub/Private/Recording/Implementations/LiveLinkUAssetRecording.cpp
+++ b/Engine/Plugins/Animation/LiveLinkHub/Source/LiveLinkHub/Private/Recording/Implementations/LiveLinkUAssetRecording.cpp
@@ -11,15 +11,16 @@
 #include "LiveLinkUAssetRecordingPlayer.h"
 #include "Misc/App.h"
 #include "Misc/FileHelper.h"
-#include "Misc/Paths.h"
 #include "Modules/ModuleManager.h"
 #include "Recording/LiveLinkHubPlaybackController.h"
-#include "Recording/LiveLinkHubRecordingController.h"
+#include "Recording/LiveLinkRecordingRangeHelpers.h"
 #include "Serialization/BufferArchive.h"
 #include "Serialization/MemoryReader.h"
 #include "Serialization/ObjectAndNameAsStringProxyArchive.h"
 #include "UObject/Package.h"
 
+#define LIVELINKHUB_FRAME_DEBUG 0
+
 ULiveLinkUAssetRecording::~ULiveLinkUAssetRecording()
 {
 	if (!IsTemplate() && !IsEngineExitRequested())
@@ -51,6 +52,15 @@ void ULiveLinkUAssetRecording::PostRename(UObject* OldOuter, const FName OldName
 	EjectAndUnload();
 }
 
+FFrameRate ULiveLinkUAssetRecording::GetGlobalFrameRate() const
+{
+	const double LastTimestamp = GetLastTimestamp();
+	const double FramesPerSecond = GetMaxFrames() / LastTimestamp;
+	const int64 FramesPerSecondInt = FMath::RoundToInt(FramesPerSecond);
+
+	return FFrameRate(static_cast<uint32>(FramesPerSecondInt), 1);
+}
+
 void ULiveLinkUAssetRecording::SaveRecordingData()
 {
 	bIsSavingRecordingData = true;
@@ -111,11 +121,16 @@ void ULiveLinkUAssetRecording::LoadRecordingData(int32 InInitialFrame, int32 InN
 		FrameFileData.Empty();
 	}
 
-	EarliestFrameToStream = StartFrame;
+	// Stream more to the left if we're nearing the end of the recording.
+	const int32 FramesPastLimit = FMath::Max(0, (StartFrame + InNumFramesToLoad) - GetMaxFrames());
+	EarliestFrameToStream = StartFrame - FramesPastLimit;
 
 	if (InitialFrameToStream != InInitialFrame)
 	{
-		bStreamingFrameChange = true;
+		// Signal to the async thread we're changing the desired frames, so it will restart itself using
+		// the most up-to-date requested frame. 
+		StreamingFrameChangeFromFrame = InitialFrameToStream;
+		UE_LOG(LogLiveLinkHub, Verbose, TEXT("Stream initial frame changed from: %d to %d "), InInitialFrame, InitialFrameToStream);
 	}
 	
 	InitialFrameToStream = InInitialFrame;
@@ -181,6 +196,7 @@ void ULiveLinkUAssetRecording::UnloadRecordingData()
 
 	FrameFileData.Empty();
 	RecordingMaxFrames = 0;
+	RecordingLastTimestamp = 0.0;
 	MaxFrameDiskSize = 0;
 	EarliestFrameToStream = 0;
 	InitialFrameToStream = 0;
@@ -199,8 +215,10 @@ void ULiveLinkUAssetRecording::UnloadRecordingData()
 	}
 }
 
-void ULiveLinkUAssetRecording::WaitForBufferedFrames(int32 InMinFrame, int32 InMaxFrame)
+bool ULiveLinkUAssetRecording::WaitForBufferedFrames(int32 InMinFrame, int32 InMaxFrame)
 {
+	DECLARE_SCOPE_CYCLE_COUNTER(TEXT("ULiveLinkUAssetRecording::WaitForBufferedFrames"), STAT_ULiveLinkUAssetRecording_WaitForBufferedFrames, STATGROUP_LiveLinkHub);
+	
 	if (AsyncStreamTask.IsValid())
 	{
 		// Max frames isn't set until after the initial load.
@@ -211,66 +229,62 @@ void ULiveLinkUAssetRecording::WaitForBufferedFrames(int32 InMinFrame, int32 InM
 
 		// Clamp the frame range to the max possible range. If the selection range extends the actual frame range
 		// then there would be nothing to load.
-		InMinFrame = FMath::Clamp(InMinFrame, 0, RecordingMaxFrames - 1);
-		InMaxFrame = FMath::Clamp(InMaxFrame, 0, RecordingMaxFrames - 1);
+		const int32 EndFrame = RecordingMaxFrames - 1;
+		InMinFrame = FMath::Clamp(InMinFrame, 0, EndFrame);
+		InMaxFrame = FMath::Clamp(InMaxFrame, 0, EndFrame);
 
 		const int32 InTotalFrames = InMaxFrame - InMinFrame + 1;
 		const TRange<int32> InRange(InMinFrame, InMaxFrame);
-		
+
 		while (true)
 		{
-			TRange<int32> BufferedFramesLocal = GetBufferedFrames();
 			if (InTotalFrames > TotalFramesToStream
-			|| BufferedFramesLocal.Contains(InRange) || AsyncStreamTask->IsDone())
+			|| IsFrameRangeBuffered(InRange)
+			|| AsyncStreamTask->IsDone())
 			{
 				break;
 			}
+			// We could potentially reduce blocking operations here by leveraging StreamingFrameChangeFromFrame and returning false
+			// if it has changed. Doing so would require both allowing it to change at this point (this is all done from the game thread)
+			// and another mechanism for ensuring the requested frame(s) we are blocking for are still delivered once loaded.
 			
 			FPlatformProcess::Sleep(0.002);
 		}
 	}
+
+	return true;
 }
 
-TRange<int32> ULiveLinkUAssetRecording::GetBufferedFrames() const
+UE::LiveLinkHub::RangeHelpers::Private::TRangeArray<int32> ULiveLinkUAssetRecording::GetBufferedFrameRanges() const
 {
+	DECLARE_SCOPE_CYCLE_COUNTER(TEXT("ULiveLinkUAssetRecording::GetBufferedFrameRanges"), STAT_ULiveLinkUAssetRecording_GetBufferedFrameRanges, STATGROUP_LiveLinkHub);
+	
 	if (bIsFullyLoaded)
 	{
-		return TRange<int32>(0, RecordingMaxFrames);
+		return { UE::LiveLinkHub::RangeHelpers::Private::MakeInclusiveRange(0, RecordingMaxFrames - 1) };
 	}
-	
-	FScopeLock Lock(&BufferedFrameMutex);
-	if (FrameFileData.Num() > 0)
-	{
-		int32 MaxStart = MIN_int32;
-		int32 MinEnd = MAX_int32;
-
-		for (const FFrameFileData& FrameData : FrameFileData)
-		{
-			const int32 RangeStart = FrameData.BufferedFrames.GetLowerBoundValue();
-			const int32 RangeEnd = FrameData.BufferedFrames.GetUpperBoundValue();
 
-			if (RangeStart > MaxStart)
-			{
-				MaxStart = RangeStart;
-			}
-
-			if (RangeEnd < MinEnd)
-			{
-				MinEnd = RangeEnd;
-			}
-		}
+	FScopeLock Lock(&BufferedFrameMutex);
+	return BufferedFrameRanges;
+}
 
-		if (MaxStart <= MinEnd)
+bool ULiveLinkUAssetRecording::IsFrameRangeBuffered(const TRange<int32>& InRange) const
+{
+	const UE::LiveLinkHub::RangeHelpers::Private::TRangeArray<int32> Ranges = GetBufferedFrameRanges();
+	for (const TRange<int32>& Range : Ranges)
+	{
+		if (Range.Contains(InRange))
 		{
-			return TRange<int32>(MaxStart, MinEnd);
+			return true;
 		}
 	}
-	
-	return TRange<int32>(0, 0);
+	return false;
 }
 
 void ULiveLinkUAssetRecording::CopyRecordingData(FLiveLinkPlaybackTracks& InOutLiveLinkPlaybackTracks) const
 {
+	DECLARE_SCOPE_CYCLE_COUNTER(TEXT("ULiveLinkUAssetRecording::CopyRecordingData"), STAT_ULiveLinkUAssetRecording_CopyRecordingData, STATGROUP_LiveLinkHub);
+	
 	FScopeLock Lock(&DataContainerMutex);
 
 	for (const TPair<FLiveLinkSubjectKey, FLiveLinkRecordingStaticDataContainer>& Pair : RecordingData.StaticData)
@@ -281,21 +295,26 @@ void ULiveLinkUAssetRecording::CopyRecordingData(FLiveLinkPlaybackTracks& InOutL
 		StaticSubjectKey.SubjectName.Name = *(StaticSubjectKey.SubjectName.ToString() + "_STATIC");
 		FLiveLinkPlaybackTrack& PlaybackTrack = InOutLiveLinkPlaybackTracks.Tracks.FindOrAdd(StaticSubjectKey);
 
-		PlaybackTrack.FrameData = TConstArrayView<TSharedPtr<FInstancedStruct>>(Pair.Value.RecordedData);
+		PlaybackTrack.FrameData = Pair.Value.RecordedData;
 		PlaybackTrack.Timestamps = TConstArrayView<double>(Pair.Value.Timestamps);
 		PlaybackTrack.LiveLinkRole = Pair.Value.Role;
 		PlaybackTrack.SubjectKey = Pair.Key;
 		PlaybackTrack.StartIndexOffset = Pair.Value.RecordedDataStartFrame;
+
+		// Don't need to record framerate for static data.
 	}
 
 	for (const TPair<FLiveLinkSubjectKey, FLiveLinkRecordingBaseDataContainer>& Pair : RecordingData.FrameData)
 	{
 		FLiveLinkPlaybackTrack& PlaybackTrack = InOutLiveLinkPlaybackTracks.Tracks.FindOrAdd(Pair.Key);
-		
-		PlaybackTrack.FrameData = TConstArrayView<TSharedPtr<FInstancedStruct>>(Pair.Value.RecordedData);
+
+		const UE::LiveLinkHub::FrameData::Private::FFrameMetaData& MetaData = FrameFileData.FindChecked(Pair.Key);
+
+		PlaybackTrack.FrameData = Pair.Value.RecordedData;
 		PlaybackTrack.Timestamps = TConstArrayView<double>(Pair.Value.Timestamps);
 		PlaybackTrack.SubjectKey = Pair.Key;
 		PlaybackTrack.StartIndexOffset = Pair.Value.RecordedDataStartFrame;
+		PlaybackTrack.LocalFrameRate = MetaData.LocalFrameRate;
 	}
 }
 
@@ -307,25 +326,30 @@ void ULiveLinkUAssetRecording::InitializeNewRecordingData(FLiveLinkUAssetRecordi
 	RecordingPreset->BuildFromClient();
 
 	int32 MaxFrames = 0;
+	double MaxLastTimestamp = 0.0;
 	for (const TTuple<FLiveLinkSubjectKey, FLiveLinkRecordingBaseDataContainer>& FrameData : RecordingData.FrameData)
 	{
 		if (FrameData.Value.Timestamps.Num() > MaxFrames)
 		{
 			MaxFrames = FrameData.Value.Timestamps.Num();
 		}
+
+		if (FrameData.Value.Timestamps.Num() > 0)
+		{
+			const double LastTimestamp = FrameData.Value.Timestamps.Last();
+			if (LastTimestamp > MaxLastTimestamp)
+			{
+				MaxLastTimestamp = LastTimestamp;
+			}
+		}
 	}
 
 	RecordingMaxFrames = MaxFrames;
+	RecordingLastTimestamp = MaxLastTimestamp;
 	
 	bIsFullyLoaded = true;
 }
 
-void ULiveLinkUAssetRecording::SetBufferedFrames(FFrameFileData& InFrameData, const TRange<int32>& InNewRange)
-{
-	FScopeLock Lock(&BufferedFrameMutex);
-	InFrameData.BufferedFrames = InNewRange;
-}
-
 void ULiveLinkUAssetRecording::SaveFrameData(FArchive* InFileWriter, const FLiveLinkSubjectKey& InSubjectKey, FLiveLinkRecordingBaseDataContainer& InBaseDataContainer)
 {
 	// This will crash if it fails -- we don't want to save invalid data.
@@ -415,18 +439,20 @@ static TAutoConsoleVariable<float> CVarLiveLinkHubDebugFrameBufferDelay(
 	TEXT("The number of seconds to wait when buffering each frame.")
 );
 
-static TAutoConsoleVariable<int32> CVarFrameBufferUpdate(
-	TEXT("LiveLinkHub.Debug.FrameBufferUpdate"),
-	5,
-	TEXT("The number of frames before updating the buffer status.")
-);
-
 void ULiveLinkUAssetRecording::LoadRecordingAsync(int32 InStartFrame, int32 InCurrentFrame, int32 InNumFramesToLoad)
 {
+	DECLARE_SCOPE_CYCLE_COUNTER(TEXT("ULiveLinkUAssetRecording::LoadRecordingAsync"), STAT_ULiveLinkUAssetRecording_LoadRecordingAsync, STATGROUP_LiveLinkHub);
+
+	using namespace UE::LiveLinkHub;
+	
+	StreamingFrameChangeFromFrame = INDEX_NONE;
+	
 	const int32 MaxPossibleFrame = RecordingMaxFrames - 1;
 	InStartFrame = FMath::Clamp(InStartFrame, 0, MaxPossibleFrame);
 	InCurrentFrame = FMath::Clamp(InCurrentFrame, 0, MaxPossibleFrame);
 	const int32 EndFrame = InStartFrame + InNumFramesToLoad - 1;
+
+	UE_LOG(LogLiveLinkHub, Verbose, TEXT("Loading recording StartFrame: %d, CurrentFrame: %d, EndFrame: %d "), InStartFrame, InCurrentFrame, EndFrame);
 	
 	ON_SCOPE_EXIT
 	{
@@ -435,14 +461,12 @@ void ULiveLinkUAssetRecording::LoadRecordingAsync(int32 InStartFrame, int32 InCu
 		bPerformedInitialLoad = true;
 	};
 	
-	if (GetBufferedFrames().Contains(TRange<int32>(InStartFrame, FMath::Min(EndFrame, MaxPossibleFrame > 0 ? MaxPossibleFrame : EndFrame))))
+	if (IsFrameRangeBuffered(RangeHelpers::Private::MakeInclusiveRange(InStartFrame, FMath::Min(EndFrame, MaxPossibleFrame > 0 ? MaxPossibleFrame : EndFrame))))
 	{
 		// All frames are already buffered.
+		UE_LOG(LogLiveLinkHub, Verbose, TEXT("Frame range already buffered"));
 		return;
 	}
-
-	DebugSleepTime = CVarLiveLinkHubDebugFrameBufferDelay.GetValueOnAnyThread();
-	ReportFrameBufferOnIteration = CVarFrameBufferUpdate.GetValueOnAnyThread();
 	
 	// Perform initial load and record entry frame file offsets.
 	const bool bInitialLoad = FrameFileData.Num() == 0;
@@ -465,7 +489,7 @@ void ULiveLinkUAssetRecording::LoadRecordingAsync(int32 InStartFrame, int32 InCu
 		for (int32 StaticIdx = 0; StaticIdx < NumStaticData; ++StaticIdx)
 		{
 			// Create framedata just to load initial static frame data. Static data doesn't require this afterward.
-			FFrameFileData TemporaryFrameData;
+			FrameData::Private::FFrameMetaData TemporaryFrameData;
 			if (!LoadInitialFrameData(TemporaryFrameData))
 			{
 				return;
@@ -473,6 +497,9 @@ void ULiveLinkUAssetRecording::LoadRecordingAsync(int32 InStartFrame, int32 InCu
 				
 			FLiveLinkRecordingStaticDataContainer& DataContainer = RecordingData.StaticData.FindChecked(*TemporaryFrameData.FrameDataSubjectKey.Get());
 			LoadFrameData(TemporaryFrameData, DataContainer, 0, 0, 1);
+
+			FScopeLock Lock(&DataContainerMutex);
+			MoveFrameDataToContainer(DataContainer, TemporaryFrameData);
 		}
 
 		// Process frame data.
@@ -482,7 +509,7 @@ void ULiveLinkUAssetRecording::LoadRecordingAsync(int32 InStartFrame, int32 InCu
 
 		for (int32 FrameIdx = 0; FrameIdx < NumFrameData; ++FrameIdx)
 		{
-			FFrameFileData KeyPosition;
+			FrameData::Private::FFrameMetaData KeyPosition;
 			if (!LoadInitialFrameData(KeyPosition))
 			{
 				return;
@@ -490,86 +517,178 @@ void ULiveLinkUAssetRecording::LoadRecordingAsync(int32 InStartFrame, int32 InCu
 
 			// Offset to the end of this block if there is multiple NumFrameData.
 			AnimationData.SetBulkDataOffset(KeyPosition.GetFrameFilePosition(KeyPosition.MaxFrames));
-			FrameFileData.Add(MoveTemp(KeyPosition));
+			FrameFileData.Add(*KeyPosition.FrameDataSubjectKey, MoveTemp(KeyPosition));
 		}
 	}
-
+	
 	// Load the required frames, either on initial load or subsequent loads.
-	for (FFrameFileData& FrameData : FrameFileData)
+	int32 CompletedTasks = 0;
+	while (true)
 	{
-		if (ensure(FrameData.FrameDataSubjectKey.IsValid()))
+		WaitIfPaused_AsyncThread();
+		
+		// Break each frame data segment into its own "task" and context switch between them, by buffering up to LiveLinkHubSettings->BufferBatchSize.
+		// The goal is to buffer the same frame numbers on multiple sources before reporting the frames as loaded, without delaying overall load
+		// too much. This way when a frame number is requested, that frame will be ready across all sources.
+		bool bHasCanceled = false;
+		for (TTuple<FLiveLinkSubjectKey, FrameData::Private::FFrameMetaData>& FrameDataKeyVal : FrameFileData)
+		{
+			FrameData::Private::FFrameMetaData& FrameData = FrameDataKeyVal.Value;
+			if (FrameData.BufferIterationData.Status <= FrameData::Private::FFrameBufferIterationData::Active)
+			{
+				if (ensure(FrameData.FrameDataSubjectKey.IsValid()))
+				{
+					FLiveLinkRecordingBaseDataContainer& DataContainer = RecordingData.FrameData.FindChecked(*FrameData.FrameDataSubjectKey.Get());
+					LoadFrameData(FrameData, DataContainer, InStartFrame, InCurrentFrame, InNumFramesToLoad);
+					if (FrameData.BufferIterationData.Status == FrameData::Private::FFrameBufferIterationData::Complete)
+					{
+						CompletedTasks++;
+					}
+				}
+				else
+				{
+					UE_LOG(LogLiveLinkHub, Error, TEXT("FrameDataSubjectKey is missing for recording %s."), *GetName());
+					CompletedTasks++;
+				}
+			}
+			if (FrameData.BufferIterationData.Status == FrameData::Private::FFrameBufferIterationData::Canceled)
+			{
+				bHasCanceled = true;
+				// Don't break, we still want to complete current iterations up to the batch size.
+			}
+		}
+
+		const bool bCompletedAllTasks = CompletedTasks >= FrameFileData.Num() || bHasCanceled;
+		
+		// Make data available on the game thread.
 		{
-			FLiveLinkRecordingBaseDataContainer& DataContainer = RecordingData.FrameData.FindChecked(*FrameData.FrameDataSubjectKey.Get());
-			LoadFrameData(FrameData, DataContainer, InStartFrame, InCurrentFrame, InNumFramesToLoad);
+			FScopeLock Lock(&DataContainerMutex);
+			for (TTuple<FLiveLinkSubjectKey, FrameData::Private::FFrameMetaData>& FrameDataKeyVal : FrameFileData)
+			{
+				FrameData::Private::FFrameMetaData& FrameData = FrameDataKeyVal.Value;
+				MoveFrameDataToContainer(RecordingData.FrameData.FindChecked(*FrameData.FrameDataSubjectKey.Get()), FrameData);
+				
+				if (bCompletedAllTasks)
+				{
+					DECLARE_SCOPE_CYCLE_COUNTER(TEXT("ULiveLinkUAssetRecording::TaskCleanup"), STAT_ULiveLinkUAssetRecording_TaskCleanup, STATGROUP_LiveLinkHub);
+					FrameData.BufferedCache.CleanCache(RangeHelpers::Private::MakeInclusiveRange(
+					InStartFrame,
+					EndFrame
+					));
+
+					FrameData.BufferIterationData.Reset();
+				}
+			}
+
+			UpdateBufferedFrames();
 		}
-		else
+		
+		if (bCompletedAllTasks)
 		{
-			UE_LOG(LogLiveLinkHub, Error, TEXT("FrameDataSubjectKey is missing for recording %s."), *GetName());
+			break;
 		}
 	}
 }
 
-bool ULiveLinkUAssetRecording::LoadInitialFrameData(FFrameFileData& OutFrameData)
+bool ULiveLinkUAssetRecording::LoadInitialFrameData(UE::LiveLinkHub::FrameData::Private::FFrameMetaData& OutFrameData)
 {
 	int32 FrameHeaderSize = 0;
 	AnimationData.ReadBulkDataPrimitive(FrameHeaderSize);
 
-	const FLiveLinkHubBulkData::FScopedBulkDataMemoryReader Reader = AnimationData.CreateBulkDataMemoryReader(FrameHeaderSize);
+	{
+		const TSharedPtr<FLiveLinkHubBulkData::FScopedBulkDataMemoryReader> Reader = AnimationData.CreateBulkDataMemoryReader(FrameHeaderSize);
 	
-	FGuid KeySource = FGuid();
-	FString KeyName;
+		FGuid KeySource = FGuid();
+		FString KeyName;
 
-	Reader.GetMemoryReader() << KeySource;
-	Reader.GetMemoryReader() << KeyName;
+		Reader->GetMemoryReader() << KeySource;
+		Reader->GetMemoryReader() << KeyName;
 			
-	OutFrameData.FrameDataSubjectKey = MakeShared<FLiveLinkSubjectKey>(KeySource, *KeyName);
+		OutFrameData.FrameDataSubjectKey = MakeShared<FLiveLinkSubjectKey>(KeySource, *KeyName);
 	
-	int32 MaxFrames = 0;
-	Reader.GetMemoryReader() << MaxFrames;
+		int32 MaxFrames = 0;
+		Reader->GetMemoryReader() << MaxFrames;
 
-	if (MaxFrames > RecordingMaxFrames)
-	{
-		RecordingMaxFrames = MaxFrames;
-	}
+		if (MaxFrames > RecordingMaxFrames)
+		{
+			RecordingMaxFrames = MaxFrames;
+		}
 	
-	OutFrameData.MaxFrames = MaxFrames;
+		OutFrameData.MaxFrames = MaxFrames;
 			
-	if (MaxFrames > 0)
-	{	
-		FString StructTypeName;
-		int32 SerializedStructureSize;
+		if (MaxFrames > 0)
+		{	
+			FString StructTypeName;
+			int32 SerializedStructureSize;
 
-		Reader.GetMemoryReader() << StructTypeName;
+			Reader->GetMemoryReader() << StructTypeName;
 
-		AnimationData.ReadBulkDataPrimitive(SerializedStructureSize);
+			AnimationData.ReadBulkDataPrimitive(SerializedStructureSize);
 
-		OutFrameData.SerializedStructureSize = SerializedStructureSize;
-		OutFrameData.RecordingStartFrameFilePosition = AnimationData.GetBulkDataOffset();
+			OutFrameData.SerializedStructureSize = SerializedStructureSize;
+			OutFrameData.RecordingStartFrameFilePosition = AnimationData.GetBulkDataOffset();
 
-		OutFrameData.LoadedStruct = FindObject<UScriptStruct>(nullptr, *StructTypeName, true);
-		if (!OutFrameData.LoadedStruct.IsValid())
-		{
-			UE_LOG(LogLiveLinkHub, Error, TEXT("Script struct type '%s' not found."), *StructTypeName);
-			return false;
+			OutFrameData.LoadedStruct = FindObject<UScriptStruct>(nullptr, *StructTypeName, true);
+			if (!OutFrameData.LoadedStruct.IsValid())
+			{
+				UE_LOG(LogLiveLinkHub, Error, TEXT("Script struct type '%s' not found."), *StructTypeName);
+				return false;
+			}
+
+			// The size on disk for each frame-- consisting of the frame index, timestamp, and frame struct data.
+			OutFrameData.FrameDiskSize = (sizeof(int32) + sizeof(double) + SerializedStructureSize);
+
+			if (OutFrameData.FrameDiskSize > MaxFrameDiskSize)
+			{
+				MaxFrameDiskSize = OutFrameData.FrameDiskSize;
+			}
 		}
+	}
 
-		// The size on disk for each frame-- consisting of the frame index, timestamp, and frame struct data.
-		OutFrameData.FrameDiskSize = (sizeof(int32) + sizeof(double) + SerializedStructureSize);
+	// Find the last timestamp, this is so we can calculate correct framerate for this track.
+	{
+		const int32 LastFrame = OutFrameData.MaxFrames - 1;
+		double Timestamp;
+		LoadTimestampFromDisk(LastFrame, OutFrameData, Timestamp);
+		OutFrameData.LastTimestamp = Timestamp;
 
-		if (OutFrameData.FrameDiskSize > MaxFrameDiskSize)
+		if (Timestamp > RecordingLastTimestamp)
 		{
-			MaxFrameDiskSize = OutFrameData.FrameDiskSize;
+			RecordingLastTimestamp = Timestamp;
 		}
 	}
 
+	// Calculate frame rate.
+	{
+		const double FramesPerSecond = OutFrameData.MaxFrames / OutFrameData.LastTimestamp;
+		const int64 FramesPerSecondInt = FMath::RoundToInt(FramesPerSecond);
+
+		OutFrameData.LocalFrameRate = FFrameRate(static_cast<uint32>(FramesPerSecondInt), 1);
+	}
+	
 	return true;
 }
 
-void ULiveLinkUAssetRecording::LoadFrameData(FFrameFileData& InFrameData, FLiveLinkRecordingBaseDataContainer& InDataContainer,
+void ULiveLinkUAssetRecording::LoadFrameData(UE::LiveLinkHub::FrameData::Private::FFrameMetaData& InFrameData, FLiveLinkRecordingBaseDataContainer& InDataContainer,
                                              int32 RequestedStartFrame, int32 RequestedInitialFrame, int32 RequestedFramesToLoad)
 {
-	bStreamingFrameChange = false;
+	DECLARE_SCOPE_CYCLE_COUNTER(TEXT("ULiveLinkUAssetRecording::LoadFrameData"), STAT_ULiveLinkUAssetRecording_LoadFrameData, STATGROUP_LiveLinkHub);
+
+	using namespace UE::LiveLinkHub;
+	
+	FrameData::Private::FFrameBufferIterationData& IterationData = InFrameData.BufferIterationData;
+
+	// First, localize the frame times to this frame data.
+	if (InFrameData.LocalFrameRate.IsValid())
+	{
+		FFrameRate RecordingFrameRate = GetGlobalFrameRate();
+		FQualifiedFrameTime StartFrameTime(FFrameTime(RequestedStartFrame), RecordingFrameRate);
+		FQualifiedFrameTime StartInitialFrameTime(FFrameTime(RequestedInitialFrame), RecordingFrameRate);
 
+		RequestedStartFrame = StartFrameTime.ConvertTo(InFrameData.LocalFrameRate).FrameNumber.Value;
+		RequestedInitialFrame = StartInitialFrameTime.ConvertTo(InFrameData.LocalFrameRate).FrameNumber.Value;
+	}
+	
 	int32 MaxFrames = InFrameData.MaxFrames;
 	if (MaxFrames > 0)
 	{
@@ -578,53 +697,222 @@ void ULiveLinkUAssetRecording::LoadFrameData(FFrameFileData& InFrameData, FLiveL
 			// Don't go past requested frames or max frames.
 			MaxFrames = FMath::Min(MaxFrames, RequestedStartFrame + RequestedFramesToLoad);
 		}
-		
-		UScriptStruct* LoadedStruct = InFrameData.LoadedStruct.Get();
-		if (LoadedStruct == nullptr)
-		{
-			UE_LOG(LogLiveLinkHub, Error, TEXT("Script struct type not found."));
-			return;
-		}
-
-		// Arrays to store the newly loaded data.
-		TArray<double> NewTimestamps;
-		TArray<TSharedPtr<FInstancedStruct>> NewRecordedData;
 
-		NewTimestamps.Reserve(MaxFrames);
-		NewRecordedData.Reserve(MaxFrames);
+		const int32 LastFrame = MaxFrames - 1;
+		check(LastFrame >= 0);
+		
+		IterationData.ForwardData.Reset();
+		IterationData.ReverseData.Reset();
 
+		const bool bIsNewIteration = IterationData.Status == FrameData::Private::FFrameBufferIterationData::New;
+		
 		// Load each frame from the initial frame, alternating right to left each frame. This creates a buffer to support
 		// scrubbing each direction and makes sure the immediate frames are loaded first.
+
+		int32 LastLoadedRightFrame, LastLoadedLeftFrame;
+		
+		if (bIsNewIteration)
+		{
+			LastLoadedRightFrame = RequestedInitialFrame;
+			LastLoadedLeftFrame = RequestedInitialFrame;
+			IterationData.Status = FrameData::Private::FFrameBufferIterationData::Active;
+		}
+		else
+		{
+			// We don't need to convert, as these are already localized.
+			LastLoadedLeftFrame = IterationData.LastLoadedLeftFrame;
+			LastLoadedRightFrame = IterationData.LastLoadedRightFrame + 1;
+		}
 		
-		int32 RightFrameIdx = RequestedInitialFrame;
-		int32 LeftFrameIdx = RequestedInitialFrame - 1; // - 1 So we don't try to load the same initial frame when alternating to the left.
-		int32 LastLoadedRightFrame = RequestedInitialFrame;
-		int32 LastLoadedLeftFrame = RequestedInitialFrame;
-		bool bLoadRight = true; // Start right -> left
+		int32 RightFrameIdx = LastLoadedRightFrame;
+		int32 LeftFrameIdx = LastLoadedLeftFrame - 1; // - 1 So we don't try to load the same initial frame when alternating to the left.
+
+		bool bLoadRight = IterationData.bLoadRight; // Start right -> left
 		
 		// We could potentially optimize this further -- such as adjusting the ratio of ahead/behind frames to buffer based on whether
 		// the recording is playing forward or reverse vs being scrubbed.
 
+		// Clear loaded frames that aren't part of the required range.
+		if (bIsNewIteration)
+		{
+			DECLARE_SCOPE_CYCLE_COUNTER(TEXT("ULiveLinkUAssetRecording::LoadFrameData::NewIteration"), STAT_ULiveLinkUAssetRecording_LoadFrameData_NewIteration, STATGROUP_LiveLinkHub);
+			
+			FScopeLock DataContainerLock(&DataContainerMutex);
+
+			TRange<int32> RangeToLoad = RangeHelpers::Private::MakeInclusiveRange(RequestedStartFrame, LastFrame);
+			TRange<int32> LoadedRange = InDataContainer.GetBufferedFrames();
+			
+			RangeHelpers::Private::TRangeArray<int32> FramesToUnload;
+			{
+				if (LoadedRange.GetLowerBoundValue() < RangeToLoad.GetLowerBoundValue())
+				{
+					TRange<int32> BeforeRange = RangeHelpers::Private::MakeInclusiveRange(LoadedRange.GetLowerBoundValue(), FMath::Min(RangeToLoad.GetLowerBoundValue(),LoadedRange.GetUpperBoundValue()));
+					FramesToUnload.Add(BeforeRange);
+				}
+
+				if (LoadedRange.GetUpperBoundValue() > RangeToLoad.GetUpperBoundValue())
+				{
+					TRange<int32> AfterRange = RangeHelpers::Private::MakeInclusiveRange(FMath::Max(RangeToLoad.GetUpperBoundValue(), LoadedRange.GetLowerBoundValue()), LoadedRange.GetUpperBoundValue());
+					FramesToUnload.Add(AfterRange);
+				}
+			}
+			
+			for (const TRange<int32>& Range : FramesToUnload)
+			{
+				MoveRangeToCache(Range, InDataContainer, InFrameData);
+			}
+
+			LoadedRange = InDataContainer.GetBufferedFrames();
+			if (!LoadedRange.Contains(InitialFrameToStream))
+			{
+				// In this case we aren't immediately loading frames, but will get to them eventually.
+				TRange<int32> Intersection = TRange<int32>::Intersection(RangeToLoad, LoadedRange);
+				MoveRangeToCache(Intersection, InDataContainer, InFrameData);
+				
+				LoadedRange = InDataContainer.GetBufferedFrames();
+			}
+
+			// If the range is partially loaded, default left/right load to loaded values to save iteration time.
+			if (!LoadedRange.IsEmpty())
+			{
+				LastLoadedLeftFrame = LoadedRange.GetLowerBoundValue();
+				LastLoadedRightFrame = LoadedRange.GetUpperBoundValue();
+				LeftFrameIdx = LastLoadedLeftFrame;
+				RightFrameIdx = LastLoadedRightFrame;
+			}
+		}
+		
+		// Verify we don't exceed the last frame, such as if this particular frame data ends before the recording's desired frame.
+		{
+			RightFrameIdx = FMath::Min(RightFrameIdx, LastFrame);
+			LastLoadedRightFrame = FMath::Min(LastLoadedRightFrame, LastFrame);
+
+			LeftFrameIdx = FMath::Min(LeftFrameIdx, LastFrame);
+			LastLoadedLeftFrame = FMath::Min(LastLoadedLeftFrame, LastFrame);
+		}
+
+		if (LeftFrameIdx == RightFrameIdx)
+		{
+			// Don't try loading the same frame.
+			RightFrameIdx++;
+		}
+		
 		auto AlternateLoadDirection = [&]()
 		{
 			bLoadRight = !bLoadRight;
 		};
 
-		int32 Iteration = 0;
-		
-		while (RightFrameIdx < MaxFrames || LeftFrameIdx >= RequestedStartFrame)
+		const int32 BufferBatchSize = GetDefault<ULiveLinkHubSettings>()->BufferBatchSize;
+
+		struct FBulkDataBatch
 		{
-			ON_SCOPE_EXIT
+			/** Reader containing unprocessed bulk data segment in memory. */
+			TSharedPtr<FLiveLinkHubBulkData::FScopedBulkDataMemoryReader> Reader;
+			/** Minimum frame of the batch. */
+			int32 MinFrame = 0;
+			/** Maximum frame of the batch. */
+			int32 MaxFrame = 0;
+			/** Total frame count of the batch. Mainly to help if a batch is empty or not. */
+			int32 Count = 0;
+			/** If a frame number is included in this batch. */
+			bool ContainsFrame(const int32 InFrame) const
 			{
-				Iteration++;
-			};
+				return InFrame >= MinFrame && InFrame <= MaxFrame && Count != 0;
+			}
+		};
+
+		// Determine ranges of disk data for left and right batches. The data must be contiguous, and will only
+		// extend until a preloaded frame is found. These batches are so we can read multiple frames from bulk data at once,
+		// saving on disk load time compared to loading the frames individually. Any data loaded here must still be processed
+		// per-frame through a serializer.
+		FBulkDataBatch RawFramesLeftBatch, RawFramesRightBatch;
+		{
+			DECLARE_SCOPE_CYCLE_COUNTER(TEXT("ULiveLinkUAssetRecording::LoadFrameData::CreateBatches"), STAT_ULiveLinkUAssetRecording_LoadFrameData_CreateBatches, STATGROUP_LiveLinkHub);
+
+			// Initial allocation, splitting evenly between left and right.
+			int32 HalfBatchSize = BufferBatchSize / 2;
+			int32 LeftBatchSize = HalfBatchSize;
+			int32 RightBatchSize = BufferBatchSize - HalfBatchSize; // Remainder goes to the right.
+
+			const int32 AvailableLeftFrames = FMath::Max(LeftFrameIdx - RequestedStartFrame, 0);
+			const int32 AvailableRightFrames = FMath::Max(MaxFrames - RightFrameIdx, 0);
+		
+			if (AvailableLeftFrames < LeftBatchSize)
+			{
+				RightBatchSize += (LeftBatchSize - AvailableLeftFrames); // Shift extra to the right.
+				LeftBatchSize = AvailableLeftFrames;
+			}
+
+			if (AvailableRightFrames < RightBatchSize)
+			{
+				LeftBatchSize += (RightBatchSize - AvailableRightFrames); // Shift extra to the left.
+				RightBatchSize = AvailableRightFrames;
+			}
+
+			// Make sure we're still within limits.
+			LeftBatchSize = FMath::Clamp(LeftBatchSize, 0, AvailableLeftFrames);
+			RightBatchSize = FMath::Clamp(RightBatchSize, 0, AvailableRightFrames);
 			
+			// Determine the min/max frames for each batch.
+			if (LeftBatchSize > 0)
+			{
+				RawFramesLeftBatch.MinFrame = LeftFrameIdx - LeftBatchSize + 1;
+				RawFramesLeftBatch.MaxFrame = LeftFrameIdx;
+
+				// Shrink the batch to the first loaded frame, left to right since it's more likely frames for the left batch
+				// will already be loaded the more to the right (center).
+				for (int32 Idx = RawFramesLeftBatch.MinFrame; Idx < RawFramesLeftBatch.MaxFrame; ++Idx)
+				{
+					if (InDataContainer.IsFrameLoaded(Idx) || InFrameData.BufferedCache.ContainsFrame(Idx))
+					{
+						const int32 Remaining = (RawFramesLeftBatch.MaxFrame - Idx);
+						LeftBatchSize -= Remaining;
+						RawFramesLeftBatch.MaxFrame = Idx - 1;
+						break;
+					}
+				}
+				if (LeftBatchSize > 0)
+				{
+					RawFramesLeftBatch.Reader = LoadRawFramesFromDisk(RawFramesLeftBatch.MinFrame, LeftBatchSize, InFrameData);
+				}
+			}
+			RawFramesLeftBatch.Count = FMath::Max(LeftBatchSize, 0);
+
+			if (RightBatchSize > 0)
+			{
+				RawFramesRightBatch.MinFrame = RightFrameIdx;
+				RawFramesRightBatch.MaxFrame = RightFrameIdx + RightBatchSize - 1;
+
+				// Shrink the batch to the first loaded frame, right to left since it's more likely frames for the right batch
+				// will already be loaded the more to the left (center).
+				for (int32 Idx = RawFramesRightBatch.MaxFrame; Idx >= RawFramesRightBatch.MinFrame; --Idx)
+				{
+					if (InDataContainer.IsFrameLoaded(Idx) || InFrameData.BufferedCache.ContainsFrame(Idx))
+					{
+						const int32 Remaining = (Idx - RawFramesRightBatch.MinFrame) + 1;
+						RightBatchSize -= Remaining;
+						RawFramesRightBatch.MinFrame = Idx + 1;
+						
+						break;
+					}
+				}
+				if (RightBatchSize > 0)
+				{
+					RawFramesRightBatch.Reader = LoadRawFramesFromDisk(RawFramesRightBatch.MinFrame, RightBatchSize, InFrameData);
+				}
+			}
+			RawFramesRightBatch.Count = FMath::Max(RightBatchSize, 0);
+		}
+
+		// Iterate through the entire range, loading frames from cache or disk. This runs until a batch cycle has completed
+		// or the entire range is loaded.
+		int32 FramesLoaded = 0;
+		while (RightFrameIdx < MaxFrames || LeftFrameIdx >= RequestedStartFrame)
+		{
 			if (bCancelStream)
 			{
 				break;
 			}
-
-			WaitIfPaused_AsyncThread();
 			
 			int32 FrameToLoad;
     
@@ -649,113 +937,216 @@ void ULiveLinkUAssetRecording::LoadFrameData(FFrameFileData& InFrameData, FLiveL
 				LastLoadedLeftFrame = FrameToLoad;
 			}
 			
-			auto InsertFrame = [&](TSharedPtr<FInstancedStruct>& InFrame, double InTimestamp, bool bCopy)
+			auto InsertFrame = [&](const TSharedPtr<FInstancedStruct>& InFrame, double InTimestamp)
 			{
-#if UE_BUILD_DEBUG
-				ensure(!NewTimestamps.Contains(InTimestamp));
+#if LIVELINKHUB_FRAME_DEBUG
+				ensureAlways(!InFrameData.BufferIterationData.ForwardData.Timestamps.Contains(InTimestamp));
+				ensureAlways(!InFrameData.BufferIterationData.ReverseData.Timestamps.Contains(InTimestamp));
+
+				ensureAlways(!InDataContainer.Timestamps.Contains(InTimestamp));
+				
+				// Additional validation to ensure timestamps / frames are loaded in the correct order.
+				for (int32 Idx = 1; Idx < InFrameData.BufferIterationData.ForwardData.Timestamps.Num(); ++Idx)
+				{
+					double LastTimestamp = InFrameData.BufferIterationData.ForwardData.Timestamps[Idx - 1];
+					double CurrentTimestamp = InFrameData.BufferIterationData.ForwardData.Timestamps[Idx];
+					ensureAlways(LastTimestamp < CurrentTimestamp);
+				}
+				for (int32 Idx = 1; Idx < InFrameData.BufferIterationData.ReverseData.Timestamps.Num(); ++Idx)
+				{
+					double LastTimestamp = InFrameData.BufferIterationData.ReverseData.Timestamps[Idx - 1];
+					double CurrentTimestamp = InFrameData.BufferIterationData.ReverseData.Timestamps[Idx];
+					ensureAlways(LastTimestamp < CurrentTimestamp);
+				}
 #endif
 				if (bLoadRight)
 				{
-					NewTimestamps.Add(InTimestamp);
-					NewRecordedData.Add(bCopy ? InFrame : MoveTemp(InFrame));
+					InFrameData.BufferIterationData.ForwardData.Timestamps.Add(InTimestamp);
+					InFrameData.BufferIterationData.ForwardData.RecordedData.Add(InFrame);
 				}
 				else
 				{
-					NewTimestamps.Insert(InTimestamp, 0);
-					NewRecordedData.Insert(bCopy ? InFrame : MoveTemp(InFrame), 0);
-				}
-				
-#if UE_BUILD_DEBUG
-				// Additional validation to ensure timestamps / frames are loaded in the correct order.
-				for (int32 Idx = 1; Idx < NewTimestamps.Num(); ++Idx)
-				{
-					double LastTimestamp = NewTimestamps[Idx - 1];
-					double CurrentTimestamp = NewTimestamps[Idx];
-					ensure(LastTimestamp < CurrentTimestamp);
+					InFrameData.BufferIterationData.ReverseData.Timestamps.Insert(InTimestamp, 0);
+					InFrameData.BufferIterationData.ReverseData.RecordedData.Insert(InFrame, 0);
 				}
-#endif
 			};
-			
-			// Don't load a frame already in memory. Do this before reading any bulk data, so we don't continually read from disk.
-			double ExistingTimestamp = 0.f;
-			if (TSharedPtr<FInstancedStruct> ExistingFrame = InDataContainer.TryGetFrame(FrameToLoad, ExistingTimestamp))
+
+			bool bIsFrameFullyLoaded = false;
 			{
-				constexpr bool bCopy = true; // Copy, as the data container could still be having its frame data pushed to animation.
-				InsertFrame(ExistingFrame, ExistingTimestamp, bCopy);
-				AlternateLoadDirection();
-				continue;
-			}
+				DECLARE_SCOPE_CYCLE_COUNTER(TEXT("ULiveLinkUAssetRecording::LoadFrameData::FindCachedFrame"), STAT_ULiveLinkUAssetRecording_LoadFrameData_FindCachedFrame, STATGROUP_LiveLinkHub);
 
-			int64 FramePosition = InFrameData.GetFrameFilePosition(FrameToLoad);
-			AnimationData.SetBulkDataOffset(FramePosition);
+				double ExistingTimestamp = 0.f;
+				if (InDataContainer.IsFrameLoaded(FrameToLoad))
+				{
+					// Frame is in-memory and part of the active range.
+					bIsFrameFullyLoaded = true;
+				}
+				else if (TSharedPtr<FInstancedStruct> ExistingFrame = InFrameData.BufferedCache.TryGetCachedFrame(FrameToLoad, ExistingTimestamp))
+				{
+					// Frame is still in memory, but was pending deletion. Move to the active range.
+					InsertFrame(ExistingFrame, ExistingTimestamp);
+					bIsFrameFullyLoaded = true;
+				}
+			}
 			
-			const FLiveLinkHubBulkData::FScopedBulkDataMemoryReader Reader = AnimationData.CreateBulkDataMemoryReader(InFrameData.FrameDiskSize);
-
-			int32 ParsedFrameIdx = 0;
-			Reader.GetMemoryReader() << ParsedFrameIdx;
-
-			// Ensure the parsed frame index matches the expected frame
-			if (!ensure(ParsedFrameIdx == FrameToLoad))
+			if (!bIsFrameFullyLoaded)
 			{
-				UE_LOG(LogLiveLinkHub, Error, TEXT("Frame index mismatch: expected %d, got %d"), FrameToLoad, ParsedFrameIdx);
-				break;
-			}
+				// Frame needs to be deserialized, and potentially loaded from disk.
+				
+				double Timestamp;
+				TSharedPtr<FInstancedStruct> DataInstancedStruct;
+				TSharedPtr<FLiveLinkHubBulkData::FScopedBulkDataMemoryReader> MemoryToUse;
+				// Check if we have preloaded the raw data into memory.
+				{
+					int64 FramePosition = 0;
+					if (RawFramesRightBatch.ContainsFrame(FrameToLoad))
+					{
+						MemoryToUse = RawFramesRightBatch.Reader;
+						check(!RawFramesLeftBatch.ContainsFrame(FrameToLoad));
+						FramePosition = InFrameData.GetRelativeFrameFilePosition(FrameToLoad - RawFramesRightBatch.MinFrame);
+					}
+					else if (RawFramesLeftBatch.ContainsFrame(FrameToLoad))
+					{
+						MemoryToUse = RawFramesLeftBatch.Reader;
+						FramePosition = InFrameData.GetRelativeFrameFilePosition(FrameToLoad - RawFramesLeftBatch.MinFrame);
+					}
+
+					if (MemoryToUse.IsValid())
+					{
+						// The position should ideally be correct, but it's possible with various caching it is off.
+						MemoryToUse->GetMemoryReader().Seek(FramePosition);
+					}
+				}
 			
-			double Timestamp = 0;
-			Reader.GetMemoryReader() << Timestamp;
-
-			// Instantiate the animation frame.
-			FObjectAndNameAsStringProxyArchive StructAr(Reader.GetMemoryReader(), false);
-			TSharedPtr<FInstancedStruct> DataInstancedStruct = MakeShared<FInstancedStruct>(LoadedStruct);
-			DataInstancedStruct->Serialize(StructAr);
-
-			InsertFrame(DataInstancedStruct, Timestamp, false);
-		
-			ensure(NewTimestamps.Num() == NewRecordedData.Num());
+				if (!ensure(LoadFrameFromDisk(FrameToLoad, InFrameData, DataInstancedStruct, Timestamp, MemoryToUse)))
+				{
+					continue;
+				}
 
-			// Update buffered frames so any tasks waiting for a certain buffer amount, as well as the UI, can update.
-			// Record every few frames to avoid constantly locking the mutex.
-			if (Iteration % ReportFrameBufferOnIteration == 0)
-			{
-				SetBufferedFrames(InFrameData, TRange<int32>(LastLoadedLeftFrame, LastLoadedRightFrame));
+				InsertFrame(DataInstancedStruct, Timestamp);
 			}
+			ensure(InDataContainer.Timestamps.Num() == InDataContainer.RecordedData.Num());
 			
+			FramesLoaded++;
 			AlternateLoadDirection();
-
-			// Test slow buffer.
-			if (DebugSleepTime > 0.f)
-			{
-				FPlatformProcess::Sleep(DebugSleepTime);
-			}
 			
-			if (bStreamingFrameChange)
+			// Break the loop when enough frames have finished constituting a batch of frames.
+			if (FramesLoaded > 0 && FramesLoaded % BufferBatchSize == 0)
 			{
-				// The requested frames to stream have changed, finish the cycle and let the async task continue with updated data.
 				break;
 			}
 		}
 
-		// Record all the frames that have been buffered.
-		const TRange<int32> NewFrameRange(LastLoadedLeftFrame, LastLoadedRightFrame);
-		if (InFrameData.BufferedFrames != NewFrameRange)
+		IterationData.bLoadRight = bLoadRight;
+		IterationData.LastLoadedLeftFrame = LastLoadedLeftFrame;
+		IterationData.LastLoadedRightFrame = LastLoadedRightFrame;
+
+		// Determine the iteration status, which will impact the overall async load loop.
+		if (RightFrameIdx >= MaxFrames && LeftFrameIdx <= RequestedStartFrame)
 		{
-			SetBufferedFrames(InFrameData, NewFrameRange);
+			IterationData.Status = FrameData::Private::FFrameBufferIterationData::Complete;
 		}
-
-		// Output the streamed data to the data container. This will unload unused frames.
+		else if (bCancelStream || StreamingFrameChangedRequested())
 		{
-			FScopeLock Lock(&DataContainerMutex);
-			InDataContainer.Timestamps = MoveTemp(NewTimestamps);
-			InDataContainer.RecordedData = MoveTemp(NewRecordedData);
-			InDataContainer.RecordedDataStartFrame = LastLoadedLeftFrame;
-
-			// This could potentially be optimized, such as outputting directly to the data container during the iteration, which could
-			// allow smoother streaming when scrubbing to a position that isn't buffered at all. However, we would need to be careful of the cost
-			// of locking the container each iteration.
+			IterationData.Status = FrameData::Private::FFrameBufferIterationData::Canceled;
 		}
 	}
 }
 
+bool ULiveLinkUAssetRecording::LoadFrameFromDisk(const int32 InFrame, const UE::LiveLinkHub::FrameData::Private::FFrameMetaData& InFrameData,
+	TSharedPtr<FInstancedStruct>& OutFrame, double& OutTimestamp, const TSharedPtr<FLiveLinkHubBulkData::FScopedBulkDataMemoryReader>& InMemory)
+{
+	DECLARE_SCOPE_CYCLE_COUNTER(TEXT("ULiveLinkUAssetRecording::LoadFrameFromDisk"), STAT_ULiveLinkUAssetRecording_LoadFrameFromDisk, STATGROUP_LiveLinkHub);
+	OutTimestamp = 0.0;
+
+	TSharedPtr<FLiveLinkHubBulkData::FScopedBulkDataMemoryReader> Reader;
+	
+	// Either use preloaded memory if passed in, or load from bulk data directly.
+	if (InMemory.IsValid())
+	{
+		Reader = InMemory;
+	}
+	else
+	{
+		const int64 FramePosition = InFrameData.GetFrameFilePosition(InFrame);
+		AnimationData.SetBulkDataOffset(FramePosition);
+		Reader = AnimationData.CreateBulkDataMemoryReader(InFrameData.FrameDiskSize);
+	}
+
+	int32 ParsedFrameIdx = 0;
+	Reader->GetMemoryReader() << ParsedFrameIdx;
+
+	// Ensure the parsed frame index matches the expected frame
+	if (!ensure(ParsedFrameIdx == InFrame))
+	{
+		UE_LOG(LogLiveLinkHub, Error, TEXT("Frame index mismatch: expected %d, got %d"), InFrame, ParsedFrameIdx);
+		return false;
+	}
+			
+	double Timestamp = 0;
+	Reader->GetMemoryReader() << Timestamp;
+
+	// Instantiate the animation frame.
+	UScriptStruct* LoadedStruct = InFrameData.LoadedStruct.Get();
+	if (LoadedStruct == nullptr)
+	{
+		UE_LOG(LogLiveLinkHub, Error, TEXT("Script struct type not found."));
+		return false;
+	}
+
+	{
+		DECLARE_SCOPE_CYCLE_COUNTER(TEXT("ULiveLinkUAssetRecording::LoadFrameFromDisk::SerializeFrame"),
+			STAT_ULiveLinkUAssetRecording_LoadFrameFromDisk_SerializeFrame, STATGROUP_LiveLinkHub);
+		FObjectAndNameAsStringProxyArchive StructAr(Reader->GetMemoryReader(), true);
+		OutFrame = MakeShared<FInstancedStruct>(LoadedStruct);
+		OutFrame->Serialize(StructAr);
+	}
+	
+	OutTimestamp = Timestamp;
+
+	return true;
+}
+
+bool ULiveLinkUAssetRecording::LoadTimestampFromDisk(const int32 InFrame, const UE::LiveLinkHub::FrameData::Private::FFrameMetaData& InFrameData, double& OutTimestamp)
+{
+	const int64 FramePosition = InFrameData.GetFrameFilePosition(InFrame);
+	AnimationData.SetBulkDataOffset(FramePosition);
+			
+	const TSharedPtr<FLiveLinkHubBulkData::FScopedBulkDataMemoryReader> Reader = AnimationData.CreateBulkDataMemoryReader(InFrameData.FrameDiskSize);
+
+	int32 ParsedFrameIdx = 0;
+	Reader->GetMemoryReader() << ParsedFrameIdx;
+
+	// Ensure the parsed frame index matches the expected frame
+	if (!ensure(ParsedFrameIdx == InFrame))
+	{
+		UE_LOG(LogLiveLinkHub, Error, TEXT("Frame index mismatch: expected %d, got %d"), InFrame, ParsedFrameIdx);
+		return false;
+	}
+			
+	Reader->GetMemoryReader() << OutTimestamp;
+
+	return true;
+}
+
+TSharedPtr<FLiveLinkHubBulkData::FScopedBulkDataMemoryReader> ULiveLinkUAssetRecording::LoadRawFramesFromDisk(const int32 InFrame,
+	const int32 InNumFrames, const UE::LiveLinkHub::FrameData::Private::FFrameMetaData& InFrameData)
+{
+	DECLARE_SCOPE_CYCLE_COUNTER(TEXT("ULiveLinkUAssetRecording::LoadRawFramesFromDisk"), STAT_ULiveLinkUAssetRecording_LoadRawFramesFromDisk, STATGROUP_LiveLinkHub);
+
+	// Seek to the beginning of the frames to load.
+	const int64 FramePosition = InFrameData.GetFrameFilePosition(InFrame);
+	AnimationData.SetBulkDataOffset(FramePosition);
+
+	// Make sure the max frames don't exceed the maximum frames for this source.
+	const int32 MaxFrames = FMath::Min(InNumFrames, InFrameData.MaxFrames - InFrame);
+	check(MaxFrames >= 1);
+
+	// Determine complete byte size to load.
+	const int32 MaxByteSize = MaxFrames * InFrameData.FrameDiskSize;
+	return AnimationData.CreateBulkDataMemoryReader(MaxByteSize);
+}
+
 void ULiveLinkUAssetRecording::EjectAndUnload()
 {
 	const FLiveLinkHubModule& LiveLinkHubModule = FModuleManager::Get().GetModuleChecked<FLiveLinkHubModule>("LiveLinkHub");
@@ -765,8 +1156,123 @@ void ULiveLinkUAssetRecording::EjectAndUnload()
 	}
 }
 
+void ULiveLinkUAssetRecording::MoveFrameDataToContainer(FLiveLinkRecordingBaseDataContainer& InDataContainer,
+	UE::LiveLinkHub::FrameData::Private::FFrameMetaData& InFrameData) const
+{
+	DECLARE_SCOPE_CYCLE_COUNTER(TEXT("ULiveLinkUAssetRecording::MoveFrameDataToContainer"), STAT_ULiveLinkUAssetRecording_MoveFrameDataToContainer, STATGROUP_LiveLinkHub);
+
+	// Don't mutex lock here, it should be locked from the calling method so multiple moves can be done under one lock.
+	
+#if LIVELINKHUB_FRAME_DEBUG
+	// Additional validation to ensure timestamps / frames are loaded in the correct order.
+	for (int32 Idx = 0; Idx < InFrameData.BufferIterationData.ForwardData.Timestamps.Num(); ++Idx)
+	{
+		ensureAlways(!InDataContainer.Timestamps.Contains(InFrameData.BufferIterationData.ForwardData.Timestamps[Idx]));
+	}
+#endif
+
+	InDataContainer.RecordedDataStartFrame = InFrameData.BufferIterationData.LastLoadedLeftFrame;
+	
+	InDataContainer.Timestamps.Insert(MoveTemp(InFrameData.BufferIterationData.ReverseData.Timestamps), 0);
+	InDataContainer.RecordedData.Insert(MoveTemp(InFrameData.BufferIterationData.ReverseData.RecordedData), 0);
+
+	InDataContainer.Timestamps.Append(MoveTemp(InFrameData.BufferIterationData.ForwardData.Timestamps));
+	InDataContainer.RecordedData.Append(MoveTemp(InFrameData.BufferIterationData.ForwardData.RecordedData));
+			
+#if LIVELINKHUB_FRAME_DEBUG
+	// Additional validation to ensure timestamps / frames are loaded in the correct order.
+	for (int32 Idx = 1; Idx < InDataContainer.Timestamps.Num(); ++Idx)
+	{
+		double LastTimestamp = InDataContainer.Timestamps[Idx - 1];
+		double CurrentTimestamp = InDataContainer.Timestamps[Idx];
+		ensureAlways(LastTimestamp < CurrentTimestamp);
+	}
+#endif
+}
+
+void ULiveLinkUAssetRecording::MoveRangeToCache(const TRange<int32>& InRange, FLiveLinkRecordingBaseDataContainer& InDataContainer,
+	UE::LiveLinkHub::FrameData::Private::FFrameMetaData& InFrameData) const
+{
+	DECLARE_SCOPE_CYCLE_COUNTER(TEXT("ULiveLinkUAssetRecording::MoveRangeToCache"), STAT_ULiveLinkUAssetRecording_MoveRangeToCache, STATGROUP_LiveLinkHub);
+	if (!InRange.IsEmpty())
+	{
+		const int32 StartIndex = InRange.GetLowerBoundValue() - InDataContainer.RecordedDataStartFrame;
+		const int32 CountToRemove = UE::LiveLinkHub::RangeHelpers::Private::GetRangeLength(InRange);
+
+		if (StartIndex >= InDataContainer.Timestamps.Num()
+			|| (StartIndex + CountToRemove > InDataContainer.Timestamps.Num())
+			|| StartIndex < 0
+			|| InDataContainer.Timestamps.IsEmpty())
+		{
+			return;
+		}
+					
+		// Move to cache. This will be unloaded later.
+		{
+			FLiveLinkRecordingBaseDataContainer HistoryContainer;
+			HistoryContainer.Timestamps.Reserve(CountToRemove);
+			HistoryContainer.RecordedData.Reserve(CountToRemove);
+
+			HistoryContainer.Timestamps.Append(&InDataContainer.Timestamps[StartIndex], CountToRemove);
+			HistoryContainer.RecordedData.Append(&InDataContainer.RecordedData[StartIndex], CountToRemove);
+			HistoryContainer.RecordedDataStartFrame = InDataContainer.RecordedDataStartFrame + StartIndex;
+
+			InFrameData.BufferedCache.FrameData.Add(MoveTemp(HistoryContainer));
+
+			InFrameData.BufferedCache.TrimCache();
+		}
+					
+		InDataContainer.Timestamps.RemoveAt(StartIndex, CountToRemove);
+		InDataContainer.RecordedData.RemoveAt(StartIndex, CountToRemove);
+
+		if (StartIndex == 0)
+		{
+			if (InDataContainer.Timestamps.Num() == 0)
+			{
+				InDataContainer.RecordedDataStartFrame = 0;
+			}
+			else
+			{
+				InDataContainer.RecordedDataStartFrame += CountToRemove;
+			}
+		}
+	}
+}
+
+void ULiveLinkUAssetRecording::UpdateBufferedFrames()
+{
+	DECLARE_SCOPE_CYCLE_COUNTER(TEXT("ULiveLinkUAssetRecording::UpdateBufferedFrames"), STAT_ULiveLinkUAssetRecording_UpdateBufferedFrames, STATGROUP_LiveLinkHub);
+
+	using namespace UE::LiveLinkHub;
+	
+	const FFrameRate GlobalFrameRate = GetGlobalFrameRate();
+	
+	FScopeLock Lock(&BufferedFrameMutex);
+
+	// Find all buffered ranges and convert them to the global frame rate. We don't want to report the local frames that are buffered, but
+	// the ranges as viewed by the recording/scrubber.
+	BufferedFrameRanges.Reset();
+	for (TTuple<FLiveLinkSubjectKey, FrameData::Private::FFrameMetaData>& FrameDataKeyVal : FrameFileData)
+	{
+		FrameData::Private::FFrameMetaData& FrameData = FrameDataKeyVal.Value;
+		RangeHelpers::Private::TRangeArray<int32> CacheBufferRanges = FrameData.BufferedCache.GetCacheBufferRanges();
+		for (TRange<int32>& Range : CacheBufferRanges)
+		{
+			Range = RangeHelpers::Private::ConvertRangeFrameRate(Range, FrameData.LocalFrameRate, GlobalFrameRate);
+		}
+		
+		TRange<int32> ActiveBufferFrames = RecordingData.FrameData.FindChecked(*FrameData.FrameDataSubjectKey.Get()).GetBufferedFrames();
+		ActiveBufferFrames = RangeHelpers::Private::ConvertRangeFrameRate(ActiveBufferFrames, FrameData.LocalFrameRate, GlobalFrameRate);
+		
+		BufferedFrameRanges.Append(MoveTemp(CacheBufferRanges));
+		BufferedFrameRanges.Add(MoveTemp(ActiveBufferFrames));
+	}
+}
+
 void ULiveLinkUAssetRecording::WaitIfPaused_AsyncThread()
 {
+	DECLARE_SCOPE_CYCLE_COUNTER(TEXT("ULiveLinkUAssetRecording::WaitIfPaused_AsyncThread"), STAT_ULiveLinkUAssetRecording_WaitIfPaused_AsyncThread, STATGROUP_LiveLinkHub);
+	
 	check(!IsInGameThread());
 	
 	if (bPauseStream)
@@ -810,8 +1316,6 @@ void ULiveLinkUAssetRecording::FLiveLinkStreamAsyncTask::DoWork()
 	int32 LastInitialFrame = -1;
 	while (LiveLinkRecording && !LiveLinkRecording->bCancelStream)
 	{
-		LiveLinkRecording->WaitIfPaused_AsyncThread();
-		
 		if (LastStartFrame != LiveLinkRecording->EarliestFrameToStream
 			|| LastTotalFrames != LiveLinkRecording->TotalFramesToStream
 			|| LastInitialFrame != LiveLinkRecording->InitialFrameToStream)
@@ -822,5 +1326,12 @@ void ULiveLinkUAssetRecording::FLiveLinkStreamAsyncTask::DoWork()
 			LiveLinkRecording->LoadRecordingAsync(LiveLinkRecording->EarliestFrameToStream,
 				LiveLinkRecording->InitialFrameToStream, LiveLinkRecording->TotalFramesToStream);
 		}
+		else
+		{
+			LiveLinkRecording->WaitIfPaused_AsyncThread();
+			FPlatformProcess::Sleep(0.002);
+		}
 	}
 }
+
+#undef LIVELINKHUB_FRAME_DEBUG
\ No newline at end of file
diff --git a/Engine/Plugins/Animation/LiveLinkHub/Source/LiveLinkHub/Private/Recording/Implementations/LiveLinkUAssetRecording.h b/Engine/Plugins/Animation/LiveLinkHub/Source/LiveLinkHub/Private/Recording/Implementations/LiveLinkUAssetRecording.h
--- a/Engine/Plugins/Animation/LiveLinkHub/Source/LiveLinkHub/Private/Recording/Implementations/LiveLinkUAssetRecording.h
+++ b/Engine/Plugins/Animation/LiveLinkHub/Source/LiveLinkHub/Private/Recording/Implementations/LiveLinkUAssetRecording.h
@@ -4,13 +4,10 @@
 
 #include "Async/AsyncWork.h"
 #include "Data/LiveLinkHubBulkData.h"
-#include "StructUtils/InstancedStruct.h"
-#include "LiveLinkPreset.h"
-#include "LiveLinkRole.h"
+#include "LiveLinkFrameData.h"
 #include "LiveLinkTypes.h"
-#include "Misc/DateTime.h"
 #include "Recording/LiveLinkRecording.h"
-#include "Templates/SubclassOf.h"
+#include "StructUtils/InstancedStruct.h"
 #include "UObject/Object.h"
 
 #include "LiveLinkUAssetRecording.generated.h"
@@ -18,113 +15,12 @@
 class FMemoryReader;
 struct FLiveLinkPlaybackTracks;
 
-/** Base data container for a recording track. */
-USTRUCT()
-struct FLiveLinkRecordingBaseDataContainer
-{
-	GENERATED_BODY()
-
-	/** SERIALIZED DATA - Timestamps for the recorded data. Each entry matches an entry in the RecordedData array. */
-	TArray<double> Timestamps;
-	
-	/**
-	 * SERIALIZED DATA - Array of either static or frame recorded for a given timestamp.
-	 * TSharedPtr used as streaming the data in may require shared access.
-	 * FSharedStruct is not used because it doesn't implement Serialize().
-	 */
-	TArray<TSharedPtr<FInstancedStruct>> RecordedData;
-
-	/** The current start frame for RecordedData. */
-	int32 RecordedDataStartFrame = 0;
-	
-	/**
-	 * Retrieve a loaded frame.
-	 *
-	 * @param InFrame The absolute frame index to load.
-	 *
-	 * @return The frame if it exists.
-	 */
-	TSharedPtr<FInstancedStruct> TryGetFrame(const int32 InFrame) const
-	{
-		if (IsFrameLoaded(InFrame))
-		{
-			const int32 RelativeFrameIdx = InFrame - RecordedDataStartFrame;
-			return RecordedData[RelativeFrameIdx];
-		}
-		return nullptr;
-	}
-
-	/**
-	 * Retrieve a loaded frame.
-	 *
-	 * @param InFrame The absolute frame index to load.
-	 * @param OutTimestamp Return the timestamp, if it exists.
-	 *
-	 * @return The frame if it exists.
-	 */
-	TSharedPtr<FInstancedStruct> TryGetFrame(const int32 InFrame, double& OutTimestamp) const
-	{
-		if (TSharedPtr<FInstancedStruct> OutFrame = TryGetFrame(InFrame))
-		{
-			const int32 RelativeFrameIdx = InFrame - RecordedDataStartFrame;
-			OutTimestamp = Timestamps[RelativeFrameIdx];
-			return OutFrame;
-		}
-		return nullptr;
-	}
-	
-	/**
-	 * Checks if a frame is currently loaded.
-	 * 
-	 * @param InFrame The absolute frame index to check.
-	 *
-	 * @return True if the frame index exists within the frame array.
-	 */
-	bool IsFrameLoaded(const int32 InFrame) const
-	{
-		return InFrame >= RecordedDataStartFrame && InFrame < RecordedDataStartFrame + RecordedData.Num();
-	}
-
-	/** Check data memory is valid and expected. */
-	void ValidateData() const
-	{
-		check(Timestamps.Num() == RecordedData.Num());
-		for (const TSharedPtr<FInstancedStruct>& InstancedStruct : RecordedData)
-		{
-			check(InstancedStruct.IsValid() && InstancedStruct->IsValid());
-		}
-	}
-};
-
-/** Container for static data. */
-USTRUCT()
-struct FLiveLinkRecordingStaticDataContainer : public FLiveLinkRecordingBaseDataContainer
-{
-	GENERATED_BODY()
-
-	/** The role of the static data being recorded. */
-	UPROPERTY()
-	TSubclassOf<ULiveLinkRole> Role = nullptr;
-};
-
-USTRUCT()
-struct FLiveLinkUAssetRecordingData
-{
-	GENERATED_BODY()
-
-	/** Length of the recording in seconds. */
-	UPROPERTY()
-	double LengthInSeconds = 0;
-
-	/** Static data encountered while recording. */
-	UPROPERTY()
-	TMap<FLiveLinkSubjectKey, FLiveLinkRecordingStaticDataContainer> StaticData;
-	
-	/** Frame data encountered while recording. */
-	UPROPERTY()
-	TMap<FLiveLinkSubjectKey, FLiveLinkRecordingBaseDataContainer> FrameData;
-};
-
+/**
+ * Asset containing all animation data stored as bulk data. This is loaded async in chunks dependent on the playhead position.
+ * Overall recording length, framerate, and frame indices are based on the maximum track length and farthest timestamp.
+ * When locating frames by indices the track will localize the frame index based on its internal framerate. All frame rates
+ * are based strictly on the number of frames and the last timestamp of the track. True frame rate is up to the client.
+ */
 UCLASS()
 class ULiveLinkUAssetRecording : public ULiveLinkRecording
 {
@@ -140,6 +36,9 @@ public:
 
 	virtual bool IsFullyLoaded() const override { return bIsFullyLoaded; }
 	virtual bool IsSavingRecordingData() const override { return bIsSavingRecordingData; }
+	virtual int32 GetMaxFrames() const override { return RecordingMaxFrames; }
+	virtual double GetLastTimestamp() const override { return RecordingLastTimestamp; }
+	virtual FFrameRate GetGlobalFrameRate() const override;
 	
 	/** Save recording data to disk. */
 	void SaveRecordingData();
@@ -151,16 +50,16 @@ public:
 	void UnloadRecordingData();
 
 	/** Block until frames are loaded. */
-	void WaitForBufferedFrames(int32 InMinFrame, int32 InMaxFrame);
-	
-	/** Return the maximum frames for this recording. */
-	int32 GetMaxFrames() const { return RecordingMaxFrames; }
+	bool WaitForBufferedFrames(int32 InMinFrame, int32 InMaxFrame);
 
 	/** The size in bytes of each animation frame. */
 	int32 GetFrameDiskSize() const { return MaxFrameDiskSize; }
 	
-	/** Return the minimum buffered frame range. */
-	TRange<int32> GetBufferedFrames() const;
+	/** Return all buffered frame ranges, including inactive. */
+	UE::LiveLinkHub::RangeHelpers::Private::TRangeArray<int32> GetBufferedFrameRanges() const;
+	
+	/** Checks if a specific frame range is buffered. */
+	bool IsFrameRangeBuffered(const TRange<int32>& InRange) const;
 
 	/** Copy the asset's loaded recording data to a format suitable for playback in live link. */
 	void CopyRecordingData(FLiveLinkPlaybackTracks& InOutLiveLinkPlaybackTracks) const;
@@ -169,34 +68,6 @@ public:
 	void InitializeNewRecordingData(FLiveLinkUAssetRecordingData&& InRecordingData, double InRecordingLengthSeconds);
 	
 private:
-	/** Frame data file information when loading from a recording file. */
-	struct FFrameFileData
-	{
-		/** The subject key used for the frame data. */
-		TSharedPtr<FLiveLinkSubjectKey> FrameDataSubjectKey;
-		/** The struct for this frame data. */
-		TWeakObjectPtr<UScriptStruct> LoadedStruct;
-		/** The position in the file recording where frame data begins. */
-		int64 RecordingStartFrameFilePosition = 0;
-		/** Maximum number of frames. */
-		int32 MaxFrames = 0;
-		/** Total size of the structure. */
-		int32 SerializedStructureSize = 0;
-		/** The size in bytes of each animation frame. */
-		int32 FrameDiskSize;
-		/** Buffered frames for this framedata. */
-		TRange<int32> BufferedFrames = TRange<int32>(0, 0);
-
-		/** Find the correct file offset based on the frame index. */
-		int64 GetFrameFilePosition(const int32 InFrameIdx) const
-		{
-			return RecordingStartFrameFilePosition + (FrameDiskSize * InFrameIdx);
-		};
-	};
-	
-	/** Update the buffered frame range. */
-	void SetBufferedFrames(FFrameFileData& InFrameData, const TRange<int32>& InNewRange);
-	
 	/** Serialize the number of frames (array size) of the BaseDataContainer to the archive. */
 	void SaveFrameData(FArchive* InFileWriter, const FLiveLinkSubjectKey& InSubjectKey, FLiveLinkRecordingBaseDataContainer& InBaseDataContainer);
 	
@@ -204,18 +75,50 @@ private:
 	void LoadRecordingAsync(int32 InStartFrame, int32 InCurrentFrame, int32 InNumFramesToLoad);
 
 	/** Initial processing on a frame, finding the correct struct and offsets. The RecordingFileReader is assumed to be at the correct position.  */
-	bool LoadInitialFrameData(FFrameFileData& OutFrameData);
+	bool LoadInitialFrameData(UE::LiveLinkHub::FrameData::Private::FFrameMetaData& OutFrameData);
 	
 	/** Load frame data to a data container. */
-	void LoadFrameData(FFrameFileData& InFrameData, FLiveLinkRecordingBaseDataContainer& InDataContainer,
+	void LoadFrameData(UE::LiveLinkHub::FrameData::Private::FFrameMetaData& InFrameData, FLiveLinkRecordingBaseDataContainer& InDataContainer,
 		int32 RequestedStartFrame, int32 RequestedInitialFrame, int32 RequestedFramesToLoad);
 
+	/**
+	 * Attempt to load a frame from bulk data.
+	 * @param InFrame The frame number to process.
+	 * @param InFrameData The frame data containing relevant information about this frame.
+	 * @param OutFrame The deserialized frame.
+	 * @param OutTimestamp The deserialized timestamp.
+	 * @param InMemory [Optional] If we are reading from memory rather than bulk data directly.
+	 */
+	bool LoadFrameFromDisk(const int32 InFrame, const UE::LiveLinkHub::FrameData::Private::FFrameMetaData& InFrameData,
+		TSharedPtr<FInstancedStruct>& OutFrame, double& OutTimestamp, const TSharedPtr<FLiveLinkHubBulkData::FScopedBulkDataMemoryReader>& InMemory = nullptr);
+
+	/** Only load the timestamp from disk. */
+	bool LoadTimestampFromDisk(const int32 InFrame, const UE::LiveLinkHub::FrameData::Private::FFrameMetaData& InFrameData, double& OutTimestamp);
+
+	/** Load multiple frames as raw data from bulk data. They still need to be passed to LoadFrameFromDisk, so they can be deserialized. */
+	TSharedPtr<FLiveLinkHubBulkData::FScopedBulkDataMemoryReader> LoadRawFramesFromDisk(const int32 InFrame,
+		const int32 InNumFrames, const UE::LiveLinkHub::FrameData::Private::FFrameMetaData& InFrameData);
+	
 	/** Eject this recording and make sure it is unloaded. */
 	void EjectAndUnload();
 
+	/** Move frame iteration data to the data container. */
+	void MoveFrameDataToContainer(FLiveLinkRecordingBaseDataContainer& InDataContainer,
+		UE::LiveLinkHub::FrameData::Private::FFrameMetaData& InFrameData) const;
+
+	/** Moves a range of frames from the container to the frame data cache. */
+	void MoveRangeToCache(const TRange<int32>& InRange, FLiveLinkRecordingBaseDataContainer& InDataContainer,
+		UE::LiveLinkHub::FrameData::Private::FFrameMetaData& InFrameData) const;
+	
+	/** Update the buffered frame range. */
+	void UpdateBufferedFrames();
+	
 	/** Make the thread wait if we are paused. */
 	void WaitIfPaused_AsyncThread();
 
+	/** If the requested streaming frame has been changed. */
+	bool StreamingFrameChangedRequested() const { return StreamingFrameChangeFromFrame != INDEX_NONE; }
+	
 	/** Signal and wait for the stream to be paused. */
 	void PauseStream();
 
@@ -265,11 +168,17 @@ private:
 	/** The animation data -- bulk data stored within this uasset. */
 	FLiveLinkHubBulkData AnimationData;
 	
-	/** The loaded frame data keys and position. */
-	TArray<FFrameFileData> FrameFileData;
+	/** The loaded frame data keys and position. Mapped by FLiveLinkSubjectKey to allow easy retrieval. */
+	TMap<FLiveLinkSubjectKey, UE::LiveLinkHub::FrameData::Private::FFrameMetaData> FrameFileData;
 
 	/** The maximum frames for this recording. */
-	int32 RecordingMaxFrames = 0;
+	std::atomic<int32> RecordingMaxFrames = 0;
+
+	/** The last time stamp of the recording.  */
+	std::atomic<double> RecordingLastTimestamp = 0.0;
+
+	/** Frames buffered, divided into ranges. */
+	UE::LiveLinkHub::RangeHelpers::Private::TRangeArray<int32> BufferedFrameRanges;
 	
 	/** The first (left most) frame to stream. */
 	int32 EarliestFrameToStream = 0;
@@ -281,7 +190,7 @@ private:
 	int32 TotalFramesToStream = 0;
 
 	/** When the streaming frame has changed, signalling the current stream task should restart. */
-	std::atomic<bool> bStreamingFrameChange = false;
+	std::atomic<int32> StreamingFrameChangeFromFrame = INDEX_NONE;
 
 	/** Signal that the stream should be canceled. */
 	std::atomic<bool> bCancelStream = false;
@@ -318,12 +227,6 @@ private:
 	
 	/** Signalled when the stream has been unpaused. */
 	FEventRef OnStreamUnpausedEvent = FEventRef(EEventMode::ManualReset);
-	
-	/** Test slow frame buffering. */
-	float DebugSleepTime = 0.f;
-
-	/** Write the frame buffer size every n iterations. */
-	int32 ReportFrameBufferOnIteration = 5;
 
 	/** If the recording is fully loaded into memory. */
 	bool bIsFullyLoaded = false;
diff --git a/Engine/Plugins/Animation/LiveLinkHub/Source/LiveLinkHub/Private/Recording/Implementations/LiveLinkUAssetRecordingPlayer.cpp b/Engine/Plugins/Animation/LiveLinkHub/Source/LiveLinkHub/Private/Recording/Implementations/LiveLinkUAssetRecordingPlayer.cpp
--- a/Engine/Plugins/Animation/LiveLinkHub/Source/LiveLinkHub/Private/Recording/Implementations/LiveLinkUAssetRecordingPlayer.cpp
+++ b/Engine/Plugins/Animation/LiveLinkHub/Source/LiveLinkHub/Private/Recording/Implementations/LiveLinkUAssetRecordingPlayer.cpp
@@ -180,16 +180,31 @@ void FLiveLinkPlaybackTrack::GetFramesUntilReverse(double InPlayhead, TArray<FLi
 	}
 }
 
-bool FLiveLinkPlaybackTrack::TryGetFrame(int32 InIndex, FLiveLinkRecordedFrame& OutFrame)
+bool FLiveLinkPlaybackTrack::TryGetFrame(const FQualifiedFrameTime& InFrameTime, FLiveLinkRecordedFrame& OutFrame)
 {
-	InIndex = GetRelativeIndex(InIndex);
-	if (InIndex >= 0 && InIndex < FrameData.Num())
+	int32 LocalizedIndex;
+	if (LocalFrameRate.IsValid())
+	{
+		// Localize the frame time relative to this track's framerate.
+		const FFrameTime LocalizedFrameTime = InFrameTime.ConvertTo(LocalFrameRate);
+		LocalizedIndex = LocalizedFrameTime.GetFrame().Value;
+	}
+	else
+	{
+		LocalizedIndex = InFrameTime.Time.GetFrame().Value;
+	}
+	
+	const int32 RelativeIndex = GetRelativeIndex(LocalizedIndex);
+	
+	if (RelativeIndex >= 0 && RelativeIndex < FrameData.Num())
 	{
-		LastReadRelativeIndex = InIndex;
+		LastReadRelativeIndex = RelativeIndex;
 		LastReadAbsoluteIndex = LastReadRelativeIndex + StartIndexOffset;
 		
 		FLiveLinkRecordedFrame FrameToPlay;
-		FrameToPlay.Data = *FrameData[InIndex];
+		FrameToPlay.Data = *FrameData[RelativeIndex];
+		check(FrameToPlay.Data.IsValid());
+		
 		FrameToPlay.SubjectKey = SubjectKey;
 		FrameToPlay.LiveLinkRole = LiveLinkRole;
 		FrameToPlay.FrameIndex = LastReadAbsoluteIndex;
@@ -263,20 +278,17 @@ TArray<FLiveLinkRecordedFrame> FLiveLinkPlaybackTracks::FetchPreviousFrames(doub
 	return PreviousFrames;
 }
 
-TArray<FLiveLinkRecordedFrame> FLiveLinkPlaybackTracks::FetchNextFramesAtIndex(int32 FrameIndex)
+TArray<FLiveLinkRecordedFrame> FLiveLinkPlaybackTracks::FetchNextFramesAtIndex(const FQualifiedFrameTime& InFrameTime)
 {
 	TArray<FLiveLinkRecordedFrame> NextFrames;
 
-	if (FrameIndex >= 0)
+	for (TTuple<FLiveLinkSubjectKey, FLiveLinkPlaybackTrack>& TrackKeyVal : Tracks)
 	{
-		for (TTuple<FLiveLinkSubjectKey, FLiveLinkPlaybackTrack>& TrackKeyVal : Tracks)
+		FLiveLinkPlaybackTrack& Track = TrackKeyVal.Value;
+		FLiveLinkRecordedFrame Frame;
+		if (Track.TryGetFrame(InFrameTime, Frame))
 		{
-			FLiveLinkPlaybackTrack& Track = TrackKeyVal.Value;
-			FLiveLinkRecordedFrame Frame;
-			if (Track.TryGetFrame(FrameIndex, Frame))
-			{
-				NextFrames.Add(MoveTemp(Frame));
-			}
+			NextFrames.Add(MoveTemp(Frame));
 		}
 	}
 
@@ -363,24 +375,27 @@ void FLiveLinkUAssetRecordingPlayer::ShutdownPlayback()
 	}
 }
 
-void FLiveLinkUAssetRecordingPlayer::StreamPlayback(int32 InFromFrame)
+bool FLiveLinkUAssetRecordingPlayer::StreamPlayback(int32 InFromFrame)
 {
 	const int32 InitialFramesToBuffer = GetNumFramesToBuffer();
 	LoadedRecording->LoadRecordingData(InFromFrame, InitialFramesToBuffer);
 
 	// Make sure there are a few frames ready.
-	LoadedRecording->WaitForBufferedFrames(InFromFrame, InFromFrame + 2);
-
-	// On initial load, the correct frame size may not be calculated prior until after waiting for the buffer,
-	// update the correct number of frames and start buffering them.
-	const int32 CurrentFramesToBuffer = GetNumFramesToBuffer();
-	if (CurrentFramesToBuffer != InitialFramesToBuffer)
+	if (LoadedRecording->WaitForBufferedFrames(InFromFrame, InFromFrame + 2))
 	{
-		LoadedRecording->LoadRecordingData(InFromFrame, CurrentFramesToBuffer);
-	}
+		// On initial load, the correct frame size may not be calculated prior until after waiting for the buffer,
+		// update the correct number of frames and start buffering them.
+		const int32 CurrentFramesToBuffer = GetNumFramesToBuffer();
+		if (CurrentFramesToBuffer != InitialFramesToBuffer)
+		{
+			LoadedRecording->LoadRecordingData(InFromFrame, CurrentFramesToBuffer);
+		}
 
-	// Take the available recording data.
-	LoadedRecording->CopyRecordingData(CurrentRecordingPlayback);
+		// Take the available recording data.
+		LoadedRecording->CopyRecordingData(CurrentRecordingPlayback);
+		return true;
+	}
+	return false;
 }
 
 int32 FLiveLinkUAssetRecordingPlayer::GetNumFramesToBuffer() const
diff --git a/Engine/Plugins/Animation/LiveLinkHub/Source/LiveLinkHub/Private/Recording/Implementations/LiveLinkUAssetRecordingPlayer.h b/Engine/Plugins/Animation/LiveLinkHub/Source/LiveLinkHub/Private/Recording/Implementations/LiveLinkUAssetRecordingPlayer.h
--- a/Engine/Plugins/Animation/LiveLinkHub/Source/LiveLinkHub/Private/Recording/Implementations/LiveLinkUAssetRecordingPlayer.h
+++ b/Engine/Plugins/Animation/LiveLinkHub/Source/LiveLinkHub/Private/Recording/Implementations/LiveLinkUAssetRecordingPlayer.h
@@ -16,7 +16,7 @@ struct FLiveLinkPlaybackTrack
 	void GetFramesUntilReverse(double InPlayhead, TArray<FLiveLinkRecordedFrame>& OutFrames);
 	
 	/** Retrieve the frame at the read index. */
-	bool TryGetFrame(int32 InIndex, FLiveLinkRecordedFrame& OutFrame);
+	bool TryGetFrame(const FQualifiedFrameTime& InFrameTime, FLiveLinkRecordedFrame& OutFrame);
 
 	/**
 	 * Convert the playhead time to a frame index.
@@ -47,7 +47,7 @@ struct FLiveLinkPlaybackTrack
 	}
 
 	/** Frame data to read. */
-	TConstArrayView<TSharedPtr<FInstancedStruct>> FrameData;
+	TArray<TSharedPtr<FInstancedStruct>> FrameData;
 	/** Timestamps for the frames in the track. */
 	TConstArrayView<double> Timestamps;
 	/** Used for static data. */
@@ -60,6 +60,8 @@ struct FLiveLinkPlaybackTrack
 	int32 LastReadAbsoluteIndex = -1;
 	/** The true index FrameData starts at. IE, if it starts at 5, then there are 5 prior frames [0..4] that aren't loaded. */
 	int32 StartIndexOffset = 0;
+	/** The frame rate of this track. Based only on the total frames and the final timestamp. */
+	FFrameRate LocalFrameRate = FFrameRate(0, 0);
 
 private:
 	/** The last timestamp recorded. */
@@ -78,7 +80,7 @@ struct FLiveLinkPlaybackTracks
 	TArray<FLiveLinkRecordedFrame> FetchPreviousFrames(double Playhead);
 	
 	/** Get the next frame(s) at the index */
-	TArray<FLiveLinkRecordedFrame> FetchNextFramesAtIndex(int32 FrameIndex);
+	TArray<FLiveLinkRecordedFrame> FetchNextFramesAtIndex(const FQualifiedFrameTime& InFrameTime);
 
 	/** Convert the playhead to a frame index */
 	int32 PlayheadToFrameIndex(double InPlayhead);
@@ -105,20 +107,29 @@ public:
 	
 	virtual TArray<FLiveLinkRecordedFrame> FetchNextFramesAtTimestamp(const FQualifiedFrameTime& InFrameTime) override
 	{
-		StreamPlayback(InFrameTime.Time.GetFrame().Value);
-		return CurrentRecordingPlayback.FetchNextFrames(InFrameTime.AsSeconds());
+		if (StreamPlayback(InFrameTime.Time.GetFrame().Value))
+		{
+			return CurrentRecordingPlayback.FetchNextFrames(InFrameTime.AsSeconds());
+		}
+		return TArray<FLiveLinkRecordedFrame>();
 	}
 
 	virtual TArray<FLiveLinkRecordedFrame> FetchPreviousFramesAtTimestamp(const FQualifiedFrameTime& InFrameTime) override
 	{
-		StreamPlayback(InFrameTime.Time.GetFrame().Value);
-		return CurrentRecordingPlayback.FetchPreviousFrames(InFrameTime.AsSeconds());
+		if (StreamPlayback(InFrameTime.Time.GetFrame().Value))
+		{
+			return CurrentRecordingPlayback.FetchPreviousFrames(InFrameTime.AsSeconds());
+		}
+		return TArray<FLiveLinkRecordedFrame>();
 	}
 
-	virtual TArray<FLiveLinkRecordedFrame> FetchNextFramesAtIndex(int32 FrameIndex) override
+	virtual TArray<FLiveLinkRecordedFrame> FetchNextFramesAtIndex(const FQualifiedFrameTime& InFrameTime) override
 	{
-		StreamPlayback(FrameIndex);
-		return CurrentRecordingPlayback.FetchNextFramesAtIndex(FrameIndex);
+		if (StreamPlayback(InFrameTime.Time.GetFrame().Value))
+		{
+			return CurrentRecordingPlayback.FetchNextFramesAtIndex(InFrameTime);
+		}
+		return TArray<FLiveLinkRecordedFrame>();
 	}
 
 	virtual void RestartPlayback(int32 InIndex) override
@@ -131,14 +142,14 @@ public:
 		return CurrentRecordingPlayback.GetInitialFrameRate();
 	}
 
-	virtual TRange<int32> GetBufferedFrames() override
+	virtual UE::LiveLinkHub::RangeHelpers::Private::TRangeArray<int32> GetBufferedFrameRanges() override
 	{
-		return LoadedRecording.IsValid() ? LoadedRecording->GetBufferedFrames() : TRange<int32>(0, 0);
+		return LoadedRecording.IsValid() ? LoadedRecording->GetBufferedFrameRanges() : UE::LiveLinkHub::RangeHelpers::Private::TRangeArray<int32>();
 	}
 
 private:
 	/** Buffer playback around a given frame. */
-	void StreamPlayback(int32 InFromFrame);
+	bool StreamPlayback(int32 InFromFrame);
 
 	/** Retrieve the total frames to buffer, based on the size the user specified in the config file. */
 	int32 GetNumFramesToBuffer() const;
diff --git a/Engine/Plugins/Animation/LiveLinkHub/Source/LiveLinkHub/Private/Recording/LiveLinkHubPlaybackController.cpp b/Engine/Plugins/Animation/LiveLinkHub/Source/LiveLinkHub/Private/Recording/LiveLinkHubPlaybackController.cpp
--- a/Engine/Plugins/Animation/LiveLinkHub/Source/LiveLinkHub/Private/Recording/LiveLinkHubPlaybackController.cpp
+++ b/Engine/Plugins/Animation/LiveLinkHub/Source/LiveLinkHub/Private/Recording/LiveLinkHubPlaybackController.cpp
@@ -73,6 +73,7 @@ FLiveLinkHubPlaybackController::FLiveLinkHubPlaybackController()
 
 FLiveLinkHubPlaybackController::~FLiveLinkHubPlaybackController()
 {
+	bIsDestructing = true;
 	Eject();
 	Stopping = true;
 	PlaybackEvent->Trigger();
@@ -115,7 +116,7 @@ TSharedRef<SWidget> FLiveLinkHubPlaybackController::MakePlaybackWidget()
 		{
 			SliderViewRange = MoveTemp(NewRange);
 		})
-		.GetBufferRange_Raw(this, &FLiveLinkHubPlaybackController::GetBufferedFrames)
+		.GetBufferRanges_Raw(this, &FLiveLinkHubPlaybackController::GetBufferedFrameRanges)
 		.GetTotalLength_Raw(this, &FLiveLinkHubPlaybackController::GetLength)
 		.GetCurrentTime_Raw(this, &FLiveLinkHubPlaybackController::GetCurrentTime)
 		.GetSelectionStartTime_Raw(this, &FLiveLinkHubPlaybackController::GetSelectionStartTime)
@@ -158,7 +159,7 @@ void FLiveLinkHubPlaybackController::ResumePlayback()
 	
 	StartTimestamp = CurrentTime.AsSeconds();
 	// Force sync so interpolation doesn't interfere if the first frame isn't the current frame
-	SyncToFrame(CurrentTime.Time.GetFrame());
+	SyncToFrame(CurrentTime);
 }
 
 void FLiveLinkHubPlaybackController::PreparePlayback(ULiveLinkRecording* InLiveLinkRecording)
@@ -204,7 +205,7 @@ void FLiveLinkHubPlaybackController::PreparePlayback(ULiveLinkRecording* InLiveL
 				RecordingToPlay->RecordingPreset->ApplyToClientLatent([this](bool)
 				{
 					bIsReady = true;
-					SyncToFrame(0); // Needed to establish connection with client
+					SyncToFrame(FQualifiedFrameTime()); // Needed to establish connection with client
 				});
 			}
 		};
@@ -402,7 +403,7 @@ void FLiveLinkHubPlaybackController::GoToTime(FQualifiedFrameTime InTime)
 	PlaybackStartTime -= TimeDouble;
 	Playhead->SetValue(InTime);
 
-	SyncToFrame(InTime.Time.GetFrame());
+	SyncToFrame(InTime);
 }
 
 FQualifiedFrameTime FLiveLinkHubPlaybackController::GetSelectionStartTime() const
@@ -429,18 +430,11 @@ FQualifiedFrameTime FLiveLinkHubPlaybackController::GetLength() const
 {
 	const FFrameRate FrameRate = GetFrameRate();
 
-	const double Length = RecordingToPlay ? RecordingToPlay->LengthInSeconds : 0.f;
-	
-	// Calculate the exact total number of frames, including fractional frames.
-	const double ExactTotalFrames = Length * FrameRate.AsDecimal();
-	const int32 TotalFrames = FMath::FloorToInt(ExactTotalFrames);
-	
-	// Calculate the fractional part and convert it to subframe precision.
-	const double Subframe = ExactTotalFrames - TotalFrames;
-	
-	const FFrameNumber LastFrameNumber(TotalFrames - 1);
+	const int32 ExactTotalFrames = RecordingToPlay ? RecordingToPlay->GetMaxFrames() : 1;
+	const int32 LastFrameIdx = ExactTotalFrames - 1;
+	const FFrameNumber LastFrameNumber(LastFrameIdx);
 
-	const FQualifiedFrameTime FrameTime(FFrameTime(LastFrameNumber, Subframe), FrameRate);
+	const FQualifiedFrameTime FrameTime(FFrameTime(LastFrameNumber), FrameRate);
 	return FrameTime;
 }
 
@@ -459,15 +453,14 @@ FFrameRate FLiveLinkHubPlaybackController::GetFrameRate() const
 {
 	if (RecordingToPlay.IsValid())
 	{
-		return RecordingToPlay->FrameRate;
+		return RecordingToPlay->GetGlobalFrameRate();
 	}
 
 	return FFrameRate(60, 1);
 }
-
-TRange<int32> FLiveLinkHubPlaybackController::GetBufferedFrames() const
+UE::LiveLinkHub::RangeHelpers::Private::TRangeArray<int32> FLiveLinkHubPlaybackController::GetBufferedFrameRanges() const
 {
-	return RecordingPlayer->GetBufferedFrames();
+	return RecordingPlayer->GetBufferedFrameRanges();
 }
 
 void FLiveLinkHubPlaybackController::Start()
@@ -545,7 +538,10 @@ uint32 FLiveLinkHubPlaybackController::Run()
 
 void FLiveLinkHubPlaybackController::OnPlaybackFinished_Internal()
 {
-	PlaybackFinishedDelegate.Broadcast();
+	if (!bIsDestructing) // Can crash otherwise, such as if we are closing the app.
+	{
+		PlaybackFinishedDelegate.Broadcast();
+	}
 }
 
 void FLiveLinkHubPlaybackController::PushSubjectData(const FLiveLinkRecordedFrame& NextFrame, bool bForceSync)
@@ -586,10 +582,13 @@ bool FLiveLinkHubPlaybackController::SyncToPlayhead()
 	return NextFrames.Num() > 0;
 }
 
-bool FLiveLinkHubPlaybackController::SyncToFrame(const FFrameNumber& InFrameNumber)
+bool FLiveLinkHubPlaybackController::SyncToFrame(const FQualifiedFrameTime& InFrameTime)
 {
-	TArray<FLiveLinkRecordedFrame> NextFrames = RecordingPlayer->FetchNextFramesAtIndex(InFrameNumber.Value);
-	
+	TArray<FLiveLinkRecordedFrame> NextFrames = RecordingPlayer->FetchNextFramesAtIndex(InFrameTime);
+	if (NextFrames.Num() == 0)
+	{
+		UE_LOG(LogLiveLinkHub, Warning, TEXT("No frame loaded for frame number %d"), InFrameTime.Time.GetFrame().Value);
+	}
 	for (const FLiveLinkRecordedFrame& NextFrame : NextFrames)
 	{
 		PushSubjectData(NextFrame, true);
diff --git a/Engine/Plugins/Animation/LiveLinkHub/Source/LiveLinkHub/Private/Recording/LiveLinkHubPlaybackController.h b/Engine/Plugins/Animation/LiveLinkHub/Source/LiveLinkHub/Private/Recording/LiveLinkHubPlaybackController.h
--- a/Engine/Plugins/Animation/LiveLinkHub/Source/LiveLinkHub/Private/Recording/LiveLinkHubPlaybackController.h
+++ b/Engine/Plugins/Animation/LiveLinkHub/Source/LiveLinkHub/Private/Recording/LiveLinkHubPlaybackController.h
@@ -14,6 +14,8 @@
 
 #include <atomic>
 
+#include "LiveLinkRecordingRangeHelpers.h"
+
 class ILiveLinkClient;
 struct FLiveLinkRecordingData;
 class FRunnableThread;
@@ -103,7 +105,7 @@ public:
 	FFrameRate GetFrameRate() const;
 
 	/** Retrieve buffered frame range. */
-	TRange<int32> GetBufferedFrames() const;
+	UE::LiveLinkHub::RangeHelpers::Private::TRangeArray<int32> GetBufferedFrameRanges() const;
 	
 	/** If the controller is ready for commands. */
 	bool IsReady() const
@@ -148,7 +150,7 @@ public:
 	}
 	
 	/** Delegate called when playback is finished (if recording is not set to loop). */
-	FSimpleMulticastDelegate& OnPlaybackFinished()
+	FTSSimpleMulticastDelegate& OnPlaybackFinished()
 	{
 		return PlaybackFinishedDelegate;
 	}
@@ -187,7 +189,7 @@ private:
 	bool SyncToPlayhead();
 
 	/** Force sync to a specific frame. */
-	bool SyncToFrame(const FFrameNumber& InFrameNumber);
+	bool SyncToFrame(const FQualifiedFrameTime& InFrameTime);
 
 	/** Checks if the current playback settings indicates the recording should restart. */
 	bool ShouldRestart() const;
@@ -208,7 +210,9 @@ private:
 	/** Whether we're currently paused. */
 	std::atomic<bool> bIsPaused = false;
 	/** If the recording is playing in reverse. */
-	std::atomic<bool>bIsReverse = false;
+	std::atomic<bool> bIsReverse = false;
+	/** If the controller is destructing. */
+	std::atomic<bool> bIsDestructing = false;
 	/** The timestamp of the animation when first playing. Can be > 0 when running in reverse. */
 	std::atomic<double> StartTimestamp = 0.f;
 	/** Indicates that we're in the process of preparing the playback. Used by the OnSourceRemoved callback to make sure we don't eject during the PreparePlayback step. */
@@ -216,7 +220,7 @@ private:
 	/** LiveLinkRecording to play.  */
 	TStrongObjectPtr<ULiveLinkRecording> RecordingToPlay;
 	/** Delegate called when a recording playback is finished (if it's not looping). */
-	FSimpleMulticastDelegate PlaybackFinishedDelegate;
+	FTSSimpleMulticastDelegate PlaybackFinishedDelegate;
 	/** Preset used to rollback the hub to its previous state after playing a recording. */
 	TStrongObjectPtr<ULiveLinkPreset> RollbackPreset;
 	/** Atomic bool keeping track of whether we should loop the playback. */
diff --git a/Engine/Plugins/Animation/LiveLinkHub/Source/LiveLinkHub/Private/Recording/LiveLinkRecording.h b/Engine/Plugins/Animation/LiveLinkHub/Source/LiveLinkHub/Private/Recording/LiveLinkRecording.h
--- a/Engine/Plugins/Animation/LiveLinkHub/Source/LiveLinkHub/Private/Recording/LiveLinkRecording.h
+++ b/Engine/Plugins/Animation/LiveLinkHub/Source/LiveLinkHub/Private/Recording/LiveLinkRecording.h
@@ -25,6 +25,15 @@ public:
 	/** True while recording data is being written to bulk data. */
 	virtual bool IsSavingRecordingData() const { return false; }
 
+	/** Return the maximum frames for this recording. */
+	virtual int32 GetMaxFrames() const { return 0; }
+
+	/** Return the last timestamp of the recording. */
+	virtual double GetLastTimestamp() const { return 0.0; }
+
+	/** Retrieve the generalized framerate of this recording. */
+	virtual FFrameRate GetGlobalFrameRate() const { return FFrameRate(); }
+
 	/** LiveLink Preset used to save the initial state of the sources and subjects at the time of recording. */
 	UPROPERTY(Instanced)
 	TObjectPtr<ULiveLinkPreset> RecordingPreset = nullptr;
diff --git a/Engine/Plugins/Animation/LiveLinkHub/Source/LiveLinkHub/Private/Recording/LiveLinkRecordingPlayer.h b/Engine/Plugins/Animation/LiveLinkHub/Source/LiveLinkHub/Private/Recording/LiveLinkRecordingPlayer.h
--- a/Engine/Plugins/Animation/LiveLinkHub/Source/LiveLinkHub/Private/Recording/LiveLinkRecordingPlayer.h
+++ b/Engine/Plugins/Animation/LiveLinkHub/Source/LiveLinkHub/Private/Recording/LiveLinkRecordingPlayer.h
@@ -2,6 +2,7 @@
 
 #pragma once
 
+#include "LiveLinkRecordingRangeHelpers.h"
 #include "LiveLinkRole.h"
 #include "LiveLinkTypes.h"
 #include "StructUtils/StructView.h"
@@ -41,7 +42,7 @@ public:
 	virtual TArray<FLiveLinkRecordedFrame> FetchPreviousFramesAtTimestamp(const FQualifiedFrameTime& InFrameTime) = 0;
 	
 	/** Fetch next frames at the provided frame index. */
-	virtual TArray<FLiveLinkRecordedFrame> FetchNextFramesAtIndex(int32 FrameIndex) = 0;
+	virtual TArray<FLiveLinkRecordedFrame> FetchNextFramesAtIndex(const FQualifiedFrameTime& InFrameTime) = 0;
 
 	/** Restart the recording from the beginning. */
 	virtual void RestartPlayback(int32 InIndex = INDEX_NONE) = 0;
@@ -50,5 +51,5 @@ public:
 	virtual FFrameRate GetInitialFramerate() = 0;
 
 	/** Retrieve the currently buffered frames. */
-	virtual TRange<int32> GetBufferedFrames() = 0;
+	virtual UE::LiveLinkHub::RangeHelpers::Private::TRangeArray<int32> GetBufferedFrameRanges() = 0;
 };
diff --git a/Engine/Plugins/Animation/LiveLinkHub/Source/LiveLinkHub/Private/Recording/LiveLinkRecordingRangeHelpers.h b/Engine/Plugins/Animation/LiveLinkHub/Source/LiveLinkHub/Private/Recording/LiveLinkRecordingRangeHelpers.h
new file mode 100644
--- /dev/null
+++ b/Engine/Plugins/Animation/LiveLinkHub/Source/LiveLinkHub/Private/Recording/LiveLinkRecordingRangeHelpers.h
@@ -0,0 +1,49 @@
+// Copyright Epic Games, Inc. All Rights Reserved.
+
+#pragma once
+
+#include "Math/Range.h"
+#include "Misc/FrameRate.h"
+#include "Misc/FrameTime.h"
+#include "Misc/QualifiedFrameTime.h"
+
+namespace UE::LiveLinkHub::RangeHelpers::Private
+{
+	/** TArray<TRange<T>> */
+	template<typename T>
+	using TRangeArray = TArray<TRange<T>>;
+	
+	/** Gets the length of the range, accounting for exclusive or inclusive ranges. */
+	template<typename T>
+	T GetRangeLength(const TRange<T>& InRange)
+	{
+		const T LowerBound = InRange.GetLowerBoundValue();
+		const T UpperBound = InRange.GetUpperBoundValue();
+
+		if (InRange.GetUpperBound().IsExclusive())
+		{
+			return UpperBound - LowerBound;
+		}
+		return UpperBound - LowerBound + 1;
+	}
+
+	/** Makes an inclusive range. */
+	template<typename T>
+	TRange<T> MakeInclusiveRange(const T InStart, const T InEnd)
+	{
+		return TRange<T>(TRangeBound<T>::Inclusive(InStart), TRangeBound<T>::Inclusive(InEnd));
+	}
+
+	/** Convert a range from one frame rate to another. */
+	template<typename T>
+	TRange<T> ConvertRangeFrameRate(const TRange<T>& InRange, const FFrameRate& InFromFrameRate, const FFrameRate& InToFrameRate)
+	{
+		const FQualifiedFrameTime FromStart(InRange.GetLowerBoundValue(), InFromFrameRate);
+		const FQualifiedFrameTime FromEnd(InRange.GetUpperBoundValue(), InFromFrameRate);
+
+		const FFrameTime ToStart = FromStart.ConvertTo(InToFrameRate);
+		const FFrameTime ToEnd = FromEnd.ConvertTo(InToFrameRate);
+
+		return TRange<T>(ToStart.GetFrame().Value, ToEnd.GetFrame().Value);
+	}
+}
diff --git a/Engine/Plugins/Animation/LiveLinkHub/Source/LiveLinkHub/Private/Settings/LiveLinkHubSettings.h b/Engine/Plugins/Animation/LiveLinkHub/Source/LiveLinkHub/Private/Settings/LiveLinkHubSettings.h
--- a/Engine/Plugins/Animation/LiveLinkHub/Source/LiveLinkHub/Private/Settings/LiveLinkHubSettings.h
+++ b/Engine/Plugins/Animation/LiveLinkHub/Source/LiveLinkHub/Private/Settings/LiveLinkHubSettings.h
@@ -30,9 +30,17 @@ public:
 	UPROPERTY(config, EditAnywhere, Category="LiveLinkHub", DisplayName = "Automatically add discovered clients")
 	bool bAutoAddDiscoveredClients = true;
 
-	/** The size in megabytes to buffer when streaming a recoding. */
+	/** The size in megabytes to buffer when streaming a recording. */
 	UPROPERTY(config, EditAnywhere, Category="LiveLinkHub", meta = (ClampMin = "1", UIMin = "1"))
 	int32 FrameBufferSizeMB = 100;
+
+	/** Number of frames to buffer at once. */
+	UPROPERTY(config, EditAnywhere, AdvancedDisplay, Category="LiveLinkHub", meta = (ClampMin = "2", UIMin = "2"))
+	int32 BufferBatchSize = 5;
+
+	/** Maximum number of frame ranges to store in history while scrubbing. Increasing can make scrubbing faster but temporarily use more memory. */
+	UPROPERTY(config, EditAnywhere, AdvancedDisplay, Category="LiveLinkHub")
+	int32 MaxBufferRangeHistory = 25;
 	
 	/** Which project settings sections to display when opening the settings viewer. */
 	UPROPERTY(config)
diff --git a/Engine/Plugins/Animation/LiveLinkHub/Source/LiveLinkHub/Private/UI/Widgets/SLiveLinkHubPlaybackWidget.cpp b/Engine/Plugins/Animation/LiveLinkHub/Source/LiveLinkHub/Private/UI/Widgets/SLiveLinkHubPlaybackWidget.cpp
--- a/Engine/Plugins/Animation/LiveLinkHub/Source/LiveLinkHub/Private/UI/Widgets/SLiveLinkHubPlaybackWidget.cpp
+++ b/Engine/Plugins/Animation/LiveLinkHub/Source/LiveLinkHub/Private/UI/Widgets/SLiveLinkHubPlaybackWidget.cpp
@@ -2,6 +2,7 @@
 
 #include "SLiveLinkHubPlaybackWidget.h"
 
+#include "Recording/LiveLinkRecordingRangeHelpers.h"
 #include "SLiveLinkHubTimeSlider.h"
 
 #include "FrameNumberNumericInterface.h"
@@ -45,7 +46,7 @@ void SLiveLinkHubPlaybackWidget::Construct(const FArguments& InArgs)
 
 	OnGetFrameRate = InArgs._GetFrameRate;
 
-	OnGetFrameBufferRange = InArgs._GetBufferRange;
+	OnGetFrameBufferRanges = InArgs._GetBufferRanges;
 
 	const TAttribute<EFrameNumberDisplayFormats> GetDisplayFormatAttr = MakeAttributeSP(this, &SLiveLinkHubPlaybackWidget::GetDisplayFormat);
 	const TAttribute<FFrameRate> GetDisplayRateAttr = MakeAttributeSP(this, &SLiveLinkHubPlaybackWidget::GetFrameRate);
@@ -63,7 +64,7 @@ void SLiveLinkHubPlaybackWidget::Construct(const FArguments& InArgs)
 		[
 			SNew(SLiveLinkHubTimeSlider)
 			.BaseArgs(SSimpleTimeSlider::FArguments()
-			.ClampRangeHighlightSize(0.15f)
+			.ClampRangeHighlightSize(0.35f)
 			.ClampRangeHighlightColor(FLinearColor::Gray.CopyWithNewOpacity(0.5f))
 			.ScrubPosition_Lambda([this]()
 			{
@@ -84,7 +85,7 @@ void SLiveLinkHubPlaybackWidget::Construct(const FArguments& InArgs)
 							}
 						})
 			)
-			.BufferRange(this, &SLiveLinkHubPlaybackWidget::GetBufferRange)
+			.BufferRange(this, &SLiveLinkHubPlaybackWidget::GetBufferRanges)
 		]
 		+SVerticalBox::Slot()
 		.AutoHeight()
@@ -400,16 +401,25 @@ TRange<double> SLiveLinkHubPlaybackWidget::GetClampRange() const
 	const FFrameRate FrameRate = GetFrameRate();
 	const double Start = GetSelectionStartTime() / FrameRate.Numerator;
 	const double End = GetSelectionEndTime() / FrameRate.Numerator;
-	return TRange<double>(Start, End);
+	return UE::LiveLinkHub::RangeHelpers::Private::MakeInclusiveRange(Start, End);
 }
 
-TRange<double> SLiveLinkHubPlaybackWidget::GetBufferRange() const
+UE::LiveLinkHub::RangeHelpers::Private::TRangeArray<double> SLiveLinkHubPlaybackWidget::GetBufferRanges() const
 {
-	check(OnGetFrameBufferRange.IsBound());
-	TRange<int32> BufferedFrames = OnGetFrameBufferRange.Execute();
-	FQualifiedFrameTime StartTime(BufferedFrames.GetLowerBound().GetValue(), GetFrameRate());
-	FQualifiedFrameTime EndTime(BufferedFrames.GetUpperBound().GetValue(), GetFrameRate());
-	return TRange<double>(StartTime.AsSeconds(), EndTime.AsSeconds());
+	using namespace UE::LiveLinkHub::RangeHelpers::Private;
+	check(OnGetFrameBufferRanges.IsBound());
+	
+	TRangeArray<double> OutputRanges;
+	
+	const TRangeArray<int32> BufferedFrameRanges = OnGetFrameBufferRanges.Execute();
+	for (const TRange<int32>& BufferedFrames : BufferedFrameRanges)
+	{
+		FQualifiedFrameTime StartTime(BufferedFrames.GetLowerBound().GetValue(), GetFrameRate());
+		FQualifiedFrameTime EndTime(BufferedFrames.GetUpperBound().GetValue(), GetFrameRate());
+		OutputRanges.Add(MakeInclusiveRange(StartTime.AsSeconds(), EndTime.AsSeconds()));
+	}
+
+	return OutputRanges;
 }
 
 bool SLiveLinkHubPlaybackWidget::IsPaused() const
diff --git a/Engine/Plugins/Animation/LiveLinkHub/Source/LiveLinkHub/Private/UI/Widgets/SLiveLinkHubPlaybackWidget.h b/Engine/Plugins/Animation/LiveLinkHub/Source/LiveLinkHub/Private/UI/Widgets/SLiveLinkHubPlaybackWidget.h
--- a/Engine/Plugins/Animation/LiveLinkHub/Source/LiveLinkHub/Private/UI/Widgets/SLiveLinkHubPlaybackWidget.h
+++ b/Engine/Plugins/Animation/LiveLinkHub/Source/LiveLinkHub/Private/UI/Widgets/SLiveLinkHubPlaybackWidget.h
@@ -5,6 +5,7 @@
 #include "FrameNumberDisplayFormat.h"
 #include "Misc/FrameRate.h"
 #include "Misc/QualifiedFrameTime.h"
+#include "Recording/LiveLinkRecordingRangeHelpers.h"
 #include "Widgets/SCompoundWidget.h"
 
 class FLiveLinkHubPlaybackController;
@@ -24,7 +25,7 @@ public:
 	DECLARE_DELEGATE_RetVal(FFrameRate, FOnGetFrame);
 	DECLARE_DELEGATE_OneParam(FOnSetDoubleRange, const TRange<double>&)
 	DECLARE_DELEGATE_RetVal(TRange<double>, FOnGetDoubleRange)
-	DECLARE_DELEGATE_RetVal(TRange<int32>, FOnGetIntRange)
+	DECLARE_DELEGATE_RetVal(UE::LiveLinkHub::RangeHelpers::Private::TRangeArray<int32>, FOnGetIntArrayRange)
 	
 	DECLARE_DELEGATE(FOnButtonPressed);
 	
@@ -71,7 +72,7 @@ public:
 	/** Set the view range (visible selection range). */
 	SLATE_EVENT(FOnGetDoubleRange, GetViewRange)
 	/** Retrieve the frame buffer range. */
-	SLATE_EVENT(FOnGetIntRange, GetBufferRange)
+	SLATE_EVENT(FOnGetIntArrayRange, GetBufferRanges)
 		
 	SLATE_END_ARGS()
 
@@ -129,7 +130,7 @@ private:
 	TRange<double> GetClampRange() const;
 
 	/** Retrieve the buffered frame range. */
-	TRange<double> GetBufferRange() const;
+	UE::LiveLinkHub::RangeHelpers::Private::TRangeArray<double> GetBufferRanges() const;
 	
 	/** Is playback paused? */
 	bool IsPaused() const;
@@ -216,7 +217,7 @@ private:
 	FOnSetDoubleRange OnSetViewRangeDelegate;
 
 	/** Retrieve the buffered frame range. */
-	FOnGetIntRange OnGetFrameBufferRange;
+	FOnGetIntArrayRange OnGetFrameBufferRanges;
 
 	/** Delegate for getting the selection start. */
 	FOnGetTime OnGetSelectionStartTimeDelegate;
diff --git a/Engine/Plugins/Animation/LiveLinkHub/Source/LiveLinkHub/Private/UI/Widgets/SLiveLinkHubTimeSlider.cpp b/Engine/Plugins/Animation/LiveLinkHub/Source/LiveLinkHub/Private/UI/Widgets/SLiveLinkHubTimeSlider.cpp
--- a/Engine/Plugins/Animation/LiveLinkHub/Source/LiveLinkHub/Private/UI/Widgets/SLiveLinkHubTimeSlider.cpp
+++ b/Engine/Plugins/Animation/LiveLinkHub/Source/LiveLinkHub/Private/UI/Widgets/SLiveLinkHubTimeSlider.cpp
@@ -4,7 +4,7 @@
 
 void SLiveLinkHubTimeSlider::Construct(const FArguments& InArgs)
 {
-	BufferRange = (InArgs._BufferRange);
+	BufferRanges = (InArgs._BufferRange);
 
 	SSimpleTimeSlider::Construct(InArgs._BaseArgs);
 }
@@ -29,28 +29,31 @@ void SLiveLinkHubTimeSlider::OnPaintExtendedSlider(bool bMirrorLabels, const FGe
 	if (LocalSequenceLength > 0)
 	{
 		FScrubRangeToScreen RangeToScreen(LocalViewRange, AllottedGeometry.GetLocalSize());
-	
-		float LeftBuffer =  RangeToScreen.InputToLocalX(BufferRange.Get().GetLowerBoundValue());
-		float RightBuffer =  RangeToScreen.InputToLocalX(BufferRange.Get().GetUpperBoundValue());
-		float Height = AllottedGeometry.GetLocalSize().Y * ClampRangeHighlightSize.Get();
-	
-		FPaintGeometry RangeGeometry;
-		if (bMirrorLabels)
-		{
-			RangeGeometry = AllottedGeometry.ToPaintGeometry(FVector2f(RightBuffer-LeftBuffer, Height), FSlateLayoutTransform(FVector2f(LeftBuffer, 0)));
-		}
-		else
+
+		for (const TRange<double>& BufferRange : BufferRanges.Get())
 		{
-			RangeGeometry = AllottedGeometry.ToPaintGeometry(FVector2f(RightBuffer-LeftBuffer, AllottedGeometry.GetLocalSize().Y), FSlateLayoutTransform(FVector2f(LeftBuffer, AllottedGeometry.GetLocalSize().Y - Height / 3.f)));
-		}
+			float LeftBuffer = RangeToScreen.InputToLocalX(BufferRange.GetLowerBoundValue());
+			float RightBuffer = RangeToScreen.InputToLocalX(BufferRange.GetUpperBoundValue());
+			float Height = AllottedGeometry.GetLocalSize().Y * ClampRangeHighlightSize.Get();
+	
+			FPaintGeometry RangeGeometry;
+			if (bMirrorLabels)
+			{
+				RangeGeometry = AllottedGeometry.ToPaintGeometry(FVector2f(RightBuffer-LeftBuffer, Height), FSlateLayoutTransform(FVector2f(LeftBuffer, 0)));
+			}
+			else
+			{
+				RangeGeometry = AllottedGeometry.ToPaintGeometry(FVector2f(RightBuffer-LeftBuffer, AllottedGeometry.GetLocalSize().Y), FSlateLayoutTransform(FVector2f(LeftBuffer, AllottedGeometry.GetLocalSize().Y - Height / 3.f)));
+			}
 
-		FSlateDrawElement::MakeBox(
-			OutDrawElements,
-			++LayerId,
-			RangeGeometry,
-			CursorBackground,
-			DrawEffects,
-			FLinearColor::White
-			);
+			FSlateDrawElement::MakeBox(
+				OutDrawElements,
+				++LayerId,
+				RangeGeometry,
+				CursorBackground,
+				DrawEffects,
+				FLinearColor::White
+				);
+		}
 	}
 }
diff --git a/Engine/Plugins/Animation/LiveLinkHub/Source/LiveLinkHub/Private/UI/Widgets/SLiveLinkHubTimeSlider.h b/Engine/Plugins/Animation/LiveLinkHub/Source/LiveLinkHub/Private/UI/Widgets/SLiveLinkHubTimeSlider.h
--- a/Engine/Plugins/Animation/LiveLinkHub/Source/LiveLinkHub/Private/UI/Widgets/SLiveLinkHubTimeSlider.h
+++ b/Engine/Plugins/Animation/LiveLinkHub/Source/LiveLinkHub/Private/UI/Widgets/SLiveLinkHubTimeSlider.h
@@ -2,6 +2,7 @@
 
 #pragma once
 
+#include "Recording/LiveLinkRecordingRangeHelpers.h"
 #include "SSimpleTimeSlider.h"
 
 /**
@@ -15,7 +16,7 @@ public:
 	SLATE_BEGIN_ARGS(SLiveLinkHubTimeSlider)
 	{}
 		/* The buffered frame size. */
-		SLATE_ATTRIBUTE(TRange<double>, BufferRange)
+		SLATE_ATTRIBUTE(UE::LiveLinkHub::RangeHelpers::Private::TRangeArray<double>, BufferRange)
 		/** SimpleTimeSlider base args. */
 		SLATE_ARGUMENT(SSimpleTimeSlider::FArguments, BaseArgs)
 	SLATE_END_ARGS()
@@ -25,7 +26,7 @@ public:
 	 * 
 	 * @param InArgs   A declaration from which to construct the widget
 	 */
-	void Construct( const FArguments& InArgs );
+	void Construct(const FArguments& InArgs);
 
 protected:
 	virtual int32 OnPaint(const FPaintArgs& Args, const FGeometry& AllottedGeometry, const FSlateRect& MyCullingRect, FSlateWindowElementList& OutDrawElements, int32 LayerId, const FWidgetStyle& InWidgetStyle, bool bParentEnabled) const override;
@@ -35,5 +36,5 @@ protected:
 	
 private:
 	/** The frame buffer range to render. */
-	TAttribute<TRange<double>> BufferRange;
+	TAttribute<UE::LiveLinkHub::RangeHelpers::Private::TRangeArray<double>> BufferRanges;
 };
\ No newline at end of file
