LiveLinkHub - Fix assert when loading a recording while it is being saved.

diff --git a/Engine/Plugins/Animation/LiveLinkHub/Source/LiveLinkHub/Private/Recording/Implementations/LiveLinkUAssetRecording.cpp b/Engine/Plugins/Animation/LiveLinkHub/Source/LiveLinkHub/Private/Recording/Implementations/LiveLinkUAssetRecording.cpp
--- a/Engine/Plugins/Animation/LiveLinkHub/Source/LiveLinkHub/Private/Recording/Implementations/LiveLinkUAssetRecording.cpp
+++ b/Engine/Plugins/Animation/LiveLinkHub/Source/LiveLinkHub/Private/Recording/Implementations/LiveLinkUAssetRecording.cpp
@@ -54,11 +54,7 @@ void ULiveLinkUAssetRecording::PostRename(UObject* OldOuter, const FName OldName
 
 FFrameRate ULiveLinkUAssetRecording::GetGlobalFrameRate() const
 {
-	const double LastTimestamp = GetLastTimestamp();
-	const double FramesPerSecond = GetMaxFrames() / LastTimestamp;
-	const int64 FramesPerSecondInt = FMath::RoundToInt(FramesPerSecond);
-
-	return FFrameRate(static_cast<uint32>(FramesPerSecondInt), 1);
+	return CalculateFrameRate(GetMaxFrames(), GetLastTimestamp());
 }
 
 void ULiveLinkUAssetRecording::SaveRecordingData()
@@ -329,19 +325,30 @@ void ULiveLinkUAssetRecording::InitializeNewRecordingData(FLiveLinkUAssetRecordi
 	double MaxLastTimestamp = 0.0;
 	for (const TTuple<FLiveLinkSubjectKey, FLiveLinkRecordingBaseDataContainer>& FrameData : RecordingData.FrameData)
 	{
-		if (FrameData.Value.Timestamps.Num() > MaxFrames)
+		const int32 LocalMaxFrames = FrameData.Value.Timestamps.Num();
+		if (LocalMaxFrames > MaxFrames)
 		{
-			MaxFrames = FrameData.Value.Timestamps.Num();
+			MaxFrames = LocalMaxFrames;
 		}
 
+		double LocalLastTimestamp = 0.0;
 		if (FrameData.Value.Timestamps.Num() > 0)
 		{
-			const double LastTimestamp = FrameData.Value.Timestamps.Last();
-			if (LastTimestamp > MaxLastTimestamp)
+			LocalLastTimestamp = FrameData.Value.Timestamps.Last();
+			if (LocalLastTimestamp > MaxLastTimestamp)
 			{
-				MaxLastTimestamp = LastTimestamp;
+				MaxLastTimestamp = LocalLastTimestamp;
 			}
 		}
+
+		// The metadata for a recording already loaded is relevant just for the frame rate.
+		UE::LiveLinkHub::FrameData::Private::FFrameMetaData MetaData;
+		{
+			MetaData.LastTimestamp = LocalLastTimestamp;
+			MetaData.MaxFrames = LocalMaxFrames;
+			MetaData.LocalFrameRate = CalculateFrameRate(LocalMaxFrames, LocalLastTimestamp);
+		}
+		FrameFileData.Add(FrameData.Key, MoveTemp(MetaData));
 	}
 
 	RecordingMaxFrames = MaxFrames;
@@ -350,6 +357,14 @@ void ULiveLinkUAssetRecording::InitializeNewRecordingData(FLiveLinkUAssetRecordi
 	bIsFullyLoaded = true;
 }
 
+FFrameRate ULiveLinkUAssetRecording::CalculateFrameRate(const int32 InMaxFrames, const double InTime)
+{
+	const double FramesPerSecond = InTime > 0.0 ? InMaxFrames / InTime : 0.0;
+	const int64 FramesPerSecondInt = FMath::RoundToInt(FramesPerSecond);
+
+	return FFrameRate(static_cast<uint32>(FramesPerSecondInt), 1);
+}
+
 void ULiveLinkUAssetRecording::SaveFrameData(FArchive* InFileWriter, const FLiveLinkSubjectKey& InSubjectKey, FLiveLinkRecordingBaseDataContainer& InBaseDataContainer)
 {
 	// This will crash if it fails -- we don't want to save invalid data.
@@ -645,7 +660,7 @@ bool ULiveLinkUAssetRecording::LoadInitialFrameData(UE::LiveLinkHub::FrameData::
 		}
 	}
 
-	// Find the last timestamp, this is so we can calculate correct framerate for this track.
+	// Find the last timestamp, this is so we can calculate the correct framerate for this track.
 	{
 		const int32 LastFrame = OutFrameData.MaxFrames - 1;
 		double Timestamp;
@@ -659,12 +674,7 @@ bool ULiveLinkUAssetRecording::LoadInitialFrameData(UE::LiveLinkHub::FrameData::
 	}
 
 	// Calculate frame rate.
-	{
-		const double FramesPerSecond = OutFrameData.MaxFrames / OutFrameData.LastTimestamp;
-		const int64 FramesPerSecondInt = FMath::RoundToInt(FramesPerSecond);
-
-		OutFrameData.LocalFrameRate = FFrameRate(static_cast<uint32>(FramesPerSecondInt), 1);
-	}
+	OutFrameData.LocalFrameRate = CalculateFrameRate(OutFrameData.MaxFrames, OutFrameData.LastTimestamp);
 	
 	return true;
 }
@@ -681,9 +691,9 @@ void ULiveLinkUAssetRecording::LoadFrameData(UE::LiveLinkHub::FrameData::Private
 	// First, localize the frame times to this frame data.
 	if (InFrameData.LocalFrameRate.IsValid())
 	{
-		FFrameRate RecordingFrameRate = GetGlobalFrameRate();
-		FQualifiedFrameTime StartFrameTime(FFrameTime(RequestedStartFrame), RecordingFrameRate);
-		FQualifiedFrameTime StartInitialFrameTime(FFrameTime(RequestedInitialFrame), RecordingFrameRate);
+		const FFrameRate RecordingFrameRate = GetGlobalFrameRate();
+		const FQualifiedFrameTime StartFrameTime(FFrameTime(RequestedStartFrame), RecordingFrameRate);
+		const FQualifiedFrameTime StartInitialFrameTime(FFrameTime(RequestedInitialFrame), RecordingFrameRate);
 
 		RequestedStartFrame = StartFrameTime.ConvertTo(InFrameData.LocalFrameRate).FrameNumber.Value;
 		RequestedInitialFrame = StartInitialFrameTime.ConvertTo(InFrameData.LocalFrameRate).FrameNumber.Value;
diff --git a/Engine/Plugins/Animation/LiveLinkHub/Source/LiveLinkHub/Private/Recording/Implementations/LiveLinkUAssetRecording.h b/Engine/Plugins/Animation/LiveLinkHub/Source/LiveLinkHub/Private/Recording/Implementations/LiveLinkUAssetRecording.h
--- a/Engine/Plugins/Animation/LiveLinkHub/Source/LiveLinkHub/Private/Recording/Implementations/LiveLinkUAssetRecording.h
+++ b/Engine/Plugins/Animation/LiveLinkHub/Source/LiveLinkHub/Private/Recording/Implementations/LiveLinkUAssetRecording.h
@@ -66,6 +66,9 @@ public:
 
 	/** Initial setup of new recording data. */
 	void InitializeNewRecordingData(FLiveLinkUAssetRecordingData&& InRecordingData, double InRecordingLengthSeconds);
+
+	/** Calculate the framerate. */
+	static FFrameRate CalculateFrameRate(const int32 InMaxFrames, const double InTime);
 	
 private:
 	/** Serialize the number of frames (array size) of the BaseDataContainer to the archive. */
@@ -74,7 +77,7 @@ private:
 	/** Initialize or update an async load. */
 	void LoadRecordingAsync(int32 InStartFrame, int32 InCurrentFrame, int32 InNumFramesToLoad);
 
-	/** Initial processing on a frame, finding the correct struct and offsets. The RecordingFileReader is assumed to be at the correct position.  */
+	/** Initial processing on a frame, finding the correct struct and offsets. The RecordingFileReader is assumed to be at the correct position. */
 	bool LoadInitialFrameData(UE::LiveLinkHub::FrameData::Private::FFrameMetaData& OutFrameData);
 	
 	/** Load frame data to a data container. */
