Add additional rates to the timecode provider in LiveLinkHub

diff --git a/Engine/Plugins/Animation/LiveLinkHub/Source/LiveLinkHub/Private/UI/Widgets/SLiveLinkTimecode.cpp b/Engine/Plugins/Animation/LiveLinkHub/Source/LiveLinkHub/Private/UI/Widgets/SLiveLinkTimecode.cpp
--- a/Engine/Plugins/Animation/LiveLinkHub/Source/LiveLinkHub/Private/UI/Widgets/SLiveLinkTimecode.cpp
+++ b/Engine/Plugins/Animation/LiveLinkHub/Source/LiveLinkHub/Private/UI/Widgets/SLiveLinkTimecode.cpp
@@ -27,10 +27,74 @@
 namespace UE::LiveLinkTimecode::Private
 {
 	static FName EnableTimecodeSourceId = TEXT("EnableTimeCodeSource");
+
 	/** We only support a preset list of timecode values + named subjects. */
+	static FName System23976fps = TEXT("SystemTime23976fps");
 	static FName System24fps = TEXT("SystemTime24fps");
+	static FName System25fps = TEXT("SystemTime25fps");
+	static FName System2997fps = TEXT("SystemTime2997fps");
 	static FName System30fps = TEXT("SystemTime30fps");
+	static FName System48fps = TEXT("SystemTime48fps");
+	static FName System50fps = TEXT("SystemTime50fps");
+	static FName System5994fps = TEXT("SystemTime5994fps");
 	static FName System60fps = TEXT("SystemTime60fps");
+
+	struct FTimecodeMenuItem
+	{
+		FFrameRate Rate;
+		FText Label;
+		FText ToolTip;
+	};
+	static const TMap<FName, FTimecodeMenuItem> StaticTimecodeMenu = []() {
+		TMap<FName, FTimecodeMenuItem> Definition = {
+			{System23976fps, {
+				FFrameRate(24000,1001),
+				LOCTEXT("LiveLinkHubTimecodeSource23976fps", "System Time (23.976 ND fps)"),
+				LOCTEXT("LiveLinkHubTimecodeSource23976fps_Tooltip", "Use a 23.976 Non-drop FPS time code based on system time.")
+			}},
+			{System24fps, {
+				FFrameRate(24,1),
+				LOCTEXT("LiveLinkHubTimecodeSource24fps", "System Time (24 fps)"),
+				LOCTEXT("LiveLinkHubTimecodeSource24fps_Tooltip", "Use a 24 FPS time code based on system time.")
+			}},
+			{System25fps, {
+				FFrameRate(25,1),
+				LOCTEXT("LiveLinkHubTimecodeSource25fps", "System Time (25 fps)"),
+				LOCTEXT("LiveLinkHubTimecodeSource25fps_Tooltip", "Use a 25 FPS time code based on system time.")
+			}},
+			{System2997fps, {
+				FFrameRate(30000,1001),
+				LOCTEXT("LiveLinkHubTimecodeSource2997fps", "System Time (29.97 ND fps)"),
+				LOCTEXT("LiveLinkHubTimecodeSource29976fps_Tooltip", "Use a 29.97 Non-drop FPS time code based on system time.")
+			}},
+			{System30fps, {
+				FFrameRate(30,1),
+				LOCTEXT("LiveLinkHubTimecodeSource30fps", "System Time (30 fps)"),
+				LOCTEXT("LiveLinkHubTimecodeSource30fps_Tooltip", "Use a 30 FPS time code based on system time.")
+			}},
+			{System48fps, {
+				FFrameRate(48,1),
+				LOCTEXT("LiveLinkHubTimecodeSource48fps", "System Time (48 fps)"),
+				LOCTEXT("LiveLinkHubTimecodeSource48fps_Tooltip", "Use a 48 FPS time code based on system time.")
+			}},
+			{System50fps, {
+				FFrameRate(50,1),
+				LOCTEXT("LiveLinkHubTimecodeSource50fps", "System Time (50 fps)"),
+				LOCTEXT("LiveLinkHubTimecodeSource50fps_Tooltip", "Use a 50 FPS time code based on system time.")
+			}},
+			{System5994fps, {
+				FFrameRate(60000, 1001),
+				LOCTEXT("LiveLinkHubTimecodeSource5994fps", "System Time (59.94 ND fps)"),
+				LOCTEXT("LiveLinkHubTimecodeSource5994fps_Tooltip", "Use a 59.94 Non-drp FPS time code based on system time.")
+			}},
+			{System60fps, {
+				FFrameRate(60, 1),
+				LOCTEXT("LiveLinkHubTimecodeSource60fps", "System Time (60 fps)"),
+				LOCTEXT("LiveLinkHubTimecodeSource60fps_Tooltip", "Use a 60 FPS time code based on system time.")
+			}}
+		};
+		return Definition;
+	}();
 }
 
 FSlateColor SLiveLinkTimecode::GetTimecodeStatusColor() const
@@ -62,20 +126,25 @@ void SLiveLinkTimecode::SendUpdatedTimecodeToEditor()
 	{
 		Settings.Source = ELiveLinkHubTimecodeSource::NotDefined;
 	}
-	else if (ActiveTimecodeSource == System24fps)
-	{
-		Settings.Source = ELiveLinkHubTimecodeSource::SystemTimeEditor;
-		Settings.DesiredFrameRate = FFrameRate(24,1);
-	}
-	else if (ActiveTimecodeSource == System30fps)
+
+	auto GetSystemTimeSource = [this]() -> TOptional<FLiveLinkHubTimecodeSettings>
 	{
-		Settings.Source = ELiveLinkHubTimecodeSource::SystemTimeEditor;
-		Settings.DesiredFrameRate = FFrameRate(30, 1);
-	}
-	else if (ActiveTimecodeSource == System60fps)
+		for (const TPair<FName,FTimecodeMenuItem>& Source : StaticTimecodeMenu)
+		{
+			if (ActiveTimecodeSource == Source.Key)
+			{
+				FLiveLinkHubTimecodeSettings Settings;
+				Settings.Source = ELiveLinkHubTimecodeSource::SystemTimeEditor;
+				Settings.DesiredFrameRate = Source.Value.Rate;
+				return Settings;
+			}
+		}
+		return {};
+	};
+
+	if (TOptional<FLiveLinkHubTimecodeSettings> SystemTimeSource = GetSystemTimeSource())
 	{
-		Settings.Source = ELiveLinkHubTimecodeSource::SystemTimeEditor;
-		Settings.DesiredFrameRate = FFrameRate(60, 1);
+		Settings = MoveTemp(*SystemTimeSource);
 	}
 	else
 	{
@@ -153,29 +222,17 @@ TSharedRef<SWidget> SLiveLinkTimecode::MakeMenu()
 			FCanExecuteAction::CreateLambda([] { return true; }),
 			FIsActionChecked::CreateLambda([this, Id] { return Id == ActiveTimecodeSource; }));
 	};
-	MenuBuilder.AddMenuEntry(
-		LOCTEXT("LiveLinkHubTimecodeSource24fps", "System Time (24 fps)"),
-		LOCTEXT("LiveLinkHubTimecodeSource24fps_Tooltip", "Use a 24 FPS time code based on system time."),
-		FSlateIcon(),
-		GenerateUIAction(UE::LiveLinkTimecode::Private::System24fps),
-		NAME_None,
-		EUserInterfaceActionType::Check);
-
-	MenuBuilder.AddMenuEntry(
-		LOCTEXT("LiveLinkHubTimecodeSource30fps", "System Time (30 fps)"),
-		LOCTEXT("LiveLinkHubTimecodeSource30fps_Tooltip", "Use a 30 FPS time code based on system time."),
-		FSlateIcon(),
-		GenerateUIAction(UE::LiveLinkTimecode::Private::System30fps),
-		NAME_None,
-		EUserInterfaceActionType::Check);
-
-	MenuBuilder.AddMenuEntry(
-		LOCTEXT("LiveLinkHubTimecodeSource60fps", "System Time (60 fps)"),
-		LOCTEXT("LiveLinkHubTimecodeSource60fps_Tooltip", "Use a 60 FPS time code based on system time."),
-		FSlateIcon(),
-		GenerateUIAction(UE::LiveLinkTimecode::Private::System60fps),
-		NAME_None,
-		EUserInterfaceActionType::Check);
+	using namespace UE::LiveLinkTimecode::Private;
+	for (const TPair<FName,FTimecodeMenuItem>& Source : StaticTimecodeMenu)
+	{
+		MenuBuilder.AddMenuEntry(
+			Source.Value.Label,
+			Source.Value.ToolTip,
+			FSlateIcon(),
+			GenerateUIAction(Source.Key),
+			NAME_None,
+			EUserInterfaceActionType::Check);
+	}
 
 	TArray<FName> Subjects;
 	WorkingClient->GetSubjectNames(Subjects);
