LiveLinkHub - Fix loading a config not always applying certain saved settings, such as Evaluation Mode of LiveLinkInputDevices.

diff --git a/Engine/Plugins/Animation/LiveLinkHub/Source/LiveLinkHub/Private/Session/LiveLinkHubSessionManager.h b/Engine/Plugins/Animation/LiveLinkHub/Source/LiveLinkHub/Private/Session/LiveLinkHubSessionManager.h
--- a/Engine/Plugins/Animation/LiveLinkHub/Source/LiveLinkHub/Private/Session/LiveLinkHubSessionManager.h
+++ b/Engine/Plugins/Animation/LiveLinkHub/Source/LiveLinkHub/Private/Session/LiveLinkHubSessionManager.h
@@ -7,10 +7,12 @@
 #include "Config/LiveLinkHubFileUtilities.h"
 #include "DesktopPlatformModule.h"
 #include "EditorDirectories.h"
+#include "Engine/Engine.h"
 #include "HAL/CriticalSection.h"
 #include "IDesktopPlatform.h"
 #include "LiveLinkHubClient.h"
 #include "LiveLinkHubLog.h"
+#include "LiveLinkSourceSettings.h"
 #include "Settings/LiveLinkHubSettings.h"
 
 #define LOCTEXT_NAMESPACE "LiveLinkHub.SessionManager"
@@ -244,6 +246,19 @@ private:
 				for (const FLiveLinkSourcePreset& SourcePreset : SessionData->Sources)
 				{
 					LiveLinkHubClient->CreateSource(SourcePreset);
+					// Ensure stored source settings persist. CreateSource will call Source->InitializeSettings, which passes in
+					// a mutable settings object. Some sources may set "default" values on the settings object overriding the
+					// saved values from the config. We want to prevent that behavior, but we still have to call InitializeSettings, because
+					// other sources may set internal values based on the current settings' values, which is behavior we want to keep.
+					if (ULiveLinkSourceSettings* PresetSettings = SourcePreset.Settings.Get())
+					{
+						if (ULiveLinkSourceSettings* CreatedSettings = LiveLinkHubClient->GetSourceSettings(SourcePreset.Guid))
+						{
+							UEngine::FCopyPropertiesForUnrelatedObjectsParams CopyParams;
+							CopyParams.bDoDelta = false;
+							UEngine::CopyPropertiesForUnrelatedObjects(PresetSettings, CreatedSettings, CopyParams);
+						}
+					}
 				}
 
 				for (const FLiveLinkSubjectPreset& SubjectPreset : SessionData->Subjects)
