Fix LiveLinkHub renaming subject to another one making it disappear in the editor

diff --git a/Engine/Plugins/Animation/LiveLinkHub/Source/LiveLinkHub/Private/LiveLinkHubSubjectSettings.cpp b/Engine/Plugins/Animation/LiveLinkHub/Source/LiveLinkHub/Private/LiveLinkHubSubjectSettings.cpp
--- a/Engine/Plugins/Animation/LiveLinkHub/Source/LiveLinkHub/Private/LiveLinkHubSubjectSettings.cpp
+++ b/Engine/Plugins/Animation/LiveLinkHub/Source/LiveLinkHub/Private/LiveLinkHubSubjectSettings.cpp
@@ -63,7 +63,7 @@ void ULiveLinkHubSubjectSettings::PostEditChangeProperty(FPropertyChangedEvent&
 	{
 		if (PreviousOutboundName != *OutboundName)
 		{
-			if (OutboundName.IsEmpty())
+			if (!ValidateOutboundName(OutboundName))
 			{
 				OutboundName = PreviousOutboundName.ToString();
 			}
@@ -101,3 +101,21 @@ void ULiveLinkHubSubjectSettings::PostEditChangeProperty(FPropertyChangedEvent&
 		LiveLinkClient->CacheSubjectSettings(Key, this);
 	}
 }
+
+bool ULiveLinkHubSubjectSettings::ValidateOutboundName(const FString& InOutboundNameCandidate) const
+{
+	if (InOutboundNameCandidate.IsEmpty())
+	{
+		return false;
+	}
+
+	if (InOutboundNameCandidate == SubjectName)
+	{
+		return true;
+	}
+
+	FLiveLinkHubClient* LiveLinkClient = static_cast<FLiveLinkHubClient*>(&IModularFeatures::Get().GetModularFeature<ILiveLinkClient>(ILiveLinkClient::ModularFeatureName));
+
+	// Can't rename to an existing subject.
+	return !LiveLinkClient->IsSubjectValid(FLiveLinkSubjectName(*InOutboundNameCandidate));
+}
diff --git a/Engine/Plugins/Animation/LiveLinkHub/Source/LiveLinkHub/Private/LiveLinkHubSubjectSettings.h b/Engine/Plugins/Animation/LiveLinkHub/Source/LiveLinkHub/Private/LiveLinkHubSubjectSettings.h
--- a/Engine/Plugins/Animation/LiveLinkHub/Source/LiveLinkHub/Private/LiveLinkHubSubjectSettings.h
+++ b/Engine/Plugins/Animation/LiveLinkHub/Source/LiveLinkHub/Private/LiveLinkHubSubjectSettings.h
@@ -33,6 +33,10 @@ public:
 	virtual void PostEditChangeProperty(FPropertyChangedEvent& PropertyChangedEvent) override;
 	//~ End UObject interface
 
+private:
+	/** Returns whether a new name candidate for the outbound name is valid. */
+	bool ValidateOutboundName(const FString& InOutboundNameCandidate) const;
+
 public:
 	/** Name of this subject. */
 	UPROPERTY(VisibleAnywhere, Category = "LiveLink")
