LiveLinkHub - Hide certain properties that aren't required in the hub, and can lead to crashes when using the Browse To button.

diff --git a/Engine/Plugins/Animation/LiveLink/Source/LiveLink/Public/LiveLinkSettings.h b/Engine/Plugins/Animation/LiveLink/Source/LiveLink/Public/LiveLinkSettings.h
--- a/Engine/Plugins/Animation/LiveLink/Source/LiveLink/Public/LiveLinkSettings.h
+++ b/Engine/Plugins/Animation/LiveLink/Source/LiveLink/Public/LiveLinkSettings.h
@@ -168,4 +168,10 @@ public:
 	float GetMessageBusHeartbeatFrequency() const { return MessageBusHeartbeatFrequency; }
 	double GetMessageBusHeartbeatTimeout() const { return MessageBusHeartbeatTimeout; }
 	double GetMessageBusTimeBeforeRemovingDeadSource() const { return MessageBusTimeBeforeRemovingInactiveSource; }
+
+	/** Retrieve the name of the protected DefaultRoleSettings property. */
+	static FName GetDefaultRoleSettingsPropertyName()
+	{
+		return GET_MEMBER_NAME_CHECKED(ULiveLinkSettings, DefaultRoleSettings);
+	}
 };
diff --git a/Engine/Plugins/Animation/LiveLinkHub/Source/LiveLinkHub/Private/LiveLinkHubModule.cpp b/Engine/Plugins/Animation/LiveLinkHub/Source/LiveLinkHub/Private/LiveLinkHubModule.cpp
--- a/Engine/Plugins/Animation/LiveLinkHub/Source/LiveLinkHub/Private/LiveLinkHubModule.cpp
+++ b/Engine/Plugins/Animation/LiveLinkHub/Source/LiveLinkHub/Private/LiveLinkHubModule.cpp
@@ -7,6 +7,7 @@
 #include "LiveLinkHubApplication.h"
 #include "LiveLinkHubLog.h"
 #include "LiveLinkHubSubjectSettings.h"
+#include "LiveLinkSettings.h"
 #include "Misc/ConfigCacheIni.h"
 #include "Modules/ModuleManager.h"
 #include "PropertyEditorModule.h"
@@ -14,8 +15,8 @@
 #include "Recording/LiveLinkHubRecordingController.h"
 #include "Settings/LiveLinkHubSettings.h"
 #include "Settings/LiveLinkHubSettingsCustomization.h"
+#include "Settings/LiveLinkSettingsCustomization.h"
 #include "Subjects/LiveLinkHubSubjectSettingsDetailsCustomization.h"
-#include "HAL/FileManager.h"
 
 #define LOCTEXT_NAMESPACE "LiveLinkHubModule"
 
@@ -61,6 +62,14 @@ void FLiveLinkHubModule::StartupModule()
 
 	PropertyModule.RegisterCustomClassLayout(ULiveLinkHubSettings::StaticClass()->GetFName(), FOnGetDetailCustomizationInstance::CreateStatic(&FLiveLinkHubSettingsCustomization::MakeInstance));
 
+	// Apply our customization for core live link settings, only if we aren't running in the full editor. We hide properties that aren't
+	// supported in a standalone application context, but are needed if loaded in the editor.
+	bUseSettingsDetailCustomization =
+		GConfig->GetBoolOrDefault(TEXT("LiveLink"), TEXT("bUseLiveLinkHubSettingsDetailCustomization"), false, GEngineIni);
+	if (bUseSettingsDetailCustomization)
+	{
+		PropertyModule.RegisterCustomClassLayout(ULiveLinkSettings::StaticClass()->GetFName(), FOnGetDetailCustomizationInstance::CreateStatic(&FLiveLinkSettingsCustomization::MakeInstance));
+	}
 	bUseSubjectSettingsDetailsCustomization =
 		GConfig->GetBoolOrDefault(TEXT("LiveLink"), TEXT("bUseLiveLinkHubSubjectSettingsDetailsCustomization"), false, GEngineIni);
 	if (bUseSubjectSettingsDetailsCustomization)
@@ -74,7 +83,11 @@ void FLiveLinkHubModule::ShutdownModule()
 	if (FPropertyEditorModule* PropertyEditorModule = FModuleManager::GetModulePtr<FPropertyEditorModule>("PropertyEditor"))
 	{
 		PropertyEditorModule->UnregisterCustomClassLayout(ULiveLinkHubSettings::StaticClass()->GetFName());
-
+		if (bUseSettingsDetailCustomization)
+		{
+			PropertyEditorModule->UnregisterCustomClassLayout(ULiveLinkSettings::StaticClass()->GetFName());
+		}
+		
 		if (bUseSubjectSettingsDetailsCustomization)
 		{
 			PropertyEditorModule->UnregisterCustomClassLayout(ULiveLinkHubSubjectSettings::StaticClass()->GetFName());
diff --git a/Engine/Plugins/Animation/LiveLinkHub/Source/LiveLinkHub/Private/LiveLinkHubModule.h b/Engine/Plugins/Animation/LiveLinkHub/Source/LiveLinkHub/Private/LiveLinkHubModule.h
--- a/Engine/Plugins/Animation/LiveLinkHub/Source/LiveLinkHub/Private/LiveLinkHubModule.h
+++ b/Engine/Plugins/Animation/LiveLinkHub/Source/LiveLinkHub/Private/LiveLinkHubModule.h
@@ -49,5 +49,8 @@ private:
 	FLiveLinkHubTicker Ticker;
 
 	/** Config-driven. True in LiveLink Hub standalone application, false in UE. */
-	bool bUseSubjectSettingsDetailsCustomization;
+	bool bUseSubjectSettingsDetailsCustomization = false;
+	
+	/** Overrides standard livelink settings customization. Config-driven. True in LiveLink Hub standalone application, false in UE. */
+	bool bUseSettingsDetailCustomization = false;
 };
diff --git a/Engine/Plugins/Animation/LiveLinkHub/Source/LiveLinkHub/Private/Settings/LiveLinkSettingsCustomization.cpp b/Engine/Plugins/Animation/LiveLinkHub/Source/LiveLinkHub/Private/Settings/LiveLinkSettingsCustomization.cpp
new file mode 100644
--- /dev/null
+++ b/Engine/Plugins/Animation/LiveLinkHub/Source/LiveLinkHub/Private/Settings/LiveLinkSettingsCustomization.cpp
@@ -0,0 +1,20 @@
+// Copyright Epic Games, Inc. All Rights Reserved.
+
+#include "LiveLinkSettingsCustomization.h"
+
+#include "DetailCategoryBuilder.h"
+#include "DetailLayoutBuilder.h"
+#include "LiveLinkSettings.h"
+
+TSharedRef<IDetailCustomization> FLiveLinkSettingsCustomization::MakeInstance()
+{
+	return MakeShared<FLiveLinkSettingsCustomization>();
+}
+
+void FLiveLinkSettingsCustomization::CustomizeDetails(IDetailLayoutBuilder& DetailBuilder)
+{
+	// These properties aren't needed when accessed through Live Link Hub.
+	// Additionally, the Browse To button isn't supported within the hub, and can lead to crashes.
+	DetailBuilder.HideProperty(ULiveLinkSettings::GetDefaultRoleSettingsPropertyName());
+	DetailBuilder.HideProperty(GET_MEMBER_NAME_CHECKED(ULiveLinkSettings, DefaultLiveLinkPreset));
+}
diff --git a/Engine/Plugins/Animation/LiveLinkHub/Source/LiveLinkHub/Private/Settings/LiveLinkSettingsCustomization.h b/Engine/Plugins/Animation/LiveLinkHub/Source/LiveLinkHub/Private/Settings/LiveLinkSettingsCustomization.h
new file mode 100644
--- /dev/null
+++ b/Engine/Plugins/Animation/LiveLinkHub/Source/LiveLinkHub/Private/Settings/LiveLinkSettingsCustomization.h
@@ -0,0 +1,18 @@
+// Copyright Epic Games, Inc. All Rights Reserved.
+
+#pragma once
+
+#include "CoreMinimal.h"
+#include "IDetailCustomization.h"
+
+/**
+ * Customization for the standard Live Link settings, specifically for Live Link Hub.
+ * If we ever need to customize these core settings under the LiveLink plugin, then this class
+ * may need to be moved there, and either extended here or made configurable to allow certain properties hidden by LLH.
+ */
+class FLiveLinkSettingsCustomization final : public IDetailCustomization
+{
+public:
+	static TSharedRef<IDetailCustomization> MakeInstance();
+	virtual void CustomizeDetails(IDetailLayoutBuilder& DetailBuilder) override;
+};
\ No newline at end of file
diff --git a/Engine/Source/Programs/LiveLinkHubEditor/Config/DefaultEngine.ini b/Engine/Source/Programs/LiveLinkHubEditor/Config/DefaultEngine.ini
--- a/Engine/Source/Programs/LiveLinkHubEditor/Config/DefaultEngine.ini
+++ b/Engine/Source/Programs/LiveLinkHubEditor/Config/DefaultEngine.ini
@@ -1,4 +1,5 @@
 [LiveLink]
+bUseLiveLinkHubSettingsDetailCustomization=true
 bUseLiveLinkHubSubjectSettingsDetailsCustomization=true
 bPanelControllerSeparateSourcesSubjects=true
 bUseModularClientReference=true
diff --git a/Engine/Source/Runtime/LiveLinkInterface/Public/LiveLinkSourceSettings.h b/Engine/Source/Runtime/LiveLinkInterface/Public/LiveLinkSourceSettings.h
--- a/Engine/Source/Runtime/LiveLinkInterface/Public/LiveLinkSourceSettings.h
+++ b/Engine/Source/Runtime/LiveLinkInterface/Public/LiveLinkSourceSettings.h
@@ -181,7 +181,7 @@ public:
 	FString ConnectionString;
 
 	/** Factory used to create the source. */
-	UPROPERTY(VisibleAnywhere, AdvancedDisplay, Category = "Settings")
+	UPROPERTY()
 	TSubclassOf<ULiveLinkSourceFactory> Factory;
 
 	/** 
