LiveLinkHub - Fix scrollbars not showing up on client list and playback list.

diff --git a/Engine/Plugins/Animation/LiveLinkHub/Source/LiveLinkHub/Private/Clients/SLiveLinkHubClientsView.h b/Engine/Plugins/Animation/LiveLinkHub/Source/LiveLinkHub/Private/Clients/SLiveLinkHubClientsView.h
--- a/Engine/Plugins/Animation/LiveLinkHub/Source/LiveLinkHub/Private/Clients/SLiveLinkHubClientsView.h
+++ b/Engine/Plugins/Animation/LiveLinkHub/Source/LiveLinkHub/Private/Clients/SLiveLinkHubClientsView.h
@@ -367,7 +367,7 @@ public:
 				]
 			]
 			+ SVerticalBox::Slot()
-			.AutoHeight()
+			.VAlign(VAlign_Fill)
 			[
 				SAssignNew(TreeView, STreeView<FClientTreeItemPtr>)
 				.TreeItemsSource(&Clients)
diff --git a/Engine/Plugins/Animation/LiveLinkHub/Source/LiveLinkHub/Private/Recording/SLiveLinkHubRecordingListView.h b/Engine/Plugins/Animation/LiveLinkHub/Source/LiveLinkHub/Private/Recording/SLiveLinkHubRecordingListView.h
--- a/Engine/Plugins/Animation/LiveLinkHub/Source/LiveLinkHub/Private/Recording/SLiveLinkHubRecordingListView.h
+++ b/Engine/Plugins/Animation/LiveLinkHub/Source/LiveLinkHub/Private/Recording/SLiveLinkHubRecordingListView.h
@@ -102,7 +102,7 @@ public:
 			[
 				SNew(SVerticalBox)
 				+ SVerticalBox::Slot()
-				.AutoHeight()
+				.VAlign(VAlign_Fill)
 				[
 					SNew(SBox)
 					.Visibility(this, &SLiveLinkHubRecordingListView::GetRecordingPickerVisibility)
