LiveLinkHub - Fix frame timing misalignment between goto frame and forward/reverse playback. Going to any frame will now take into account both the local track framerate for the initial conversion and the nearest recorded timestamp. This results in more correct and smoother animation between scrubbing and playback, especially in cases where the initial recorded frames have a > 0 timestamp.

diff --git a/Engine/Plugins/Animation/LiveLinkHub/Source/LiveLinkHub/Private/Recording/Implementations/LiveLinkUAssetRecordingPlayer.cpp b/Engine/Plugins/Animation/LiveLinkHub/Source/LiveLinkHub/Private/Recording/Implementations/LiveLinkUAssetRecordingPlayer.cpp
--- a/Engine/Plugins/Animation/LiveLinkHub/Source/LiveLinkHub/Private/Recording/Implementations/LiveLinkUAssetRecordingPlayer.cpp
+++ b/Engine/Plugins/Animation/LiveLinkHub/Source/LiveLinkHub/Private/Recording/Implementations/LiveLinkUAssetRecordingPlayer.cpp
@@ -105,8 +105,9 @@ private:
 	}
 };
 
-void FLiveLinkPlaybackTrack::GetFramesUntil(double InPlayhead, TArray<FLiveLinkRecordedFrame>& OutFrames)
+void FLiveLinkPlaybackTrack::GetFramesUntil(const FQualifiedFrameTime& InFrameTime, TArray<FLiveLinkRecordedFrame>& OutFrames)
 {
+	const double InTimeSeconds = InFrameTime.AsSeconds();
 	for (FLiveLinkPlaybackTrackForwardIterator It = FLiveLinkPlaybackTrackForwardIterator(*this, GetRelativeIndex(LastReadAbsoluteIndex)); It; ++It)
 	{
 		const double FrameTimestamp = It.FrameTimestamp();
@@ -118,7 +119,7 @@ void FLiveLinkPlaybackTrack::GetFramesUntil(double InPlayhead, TArray<FLiveLinkR
 			continue;
 		}
 		
-		if (FrameTimestamp > InPlayhead)
+		if (FrameTimestamp > InTimeSeconds)
 		{
 			break;
 		}
@@ -137,7 +138,7 @@ void FLiveLinkPlaybackTrack::GetFramesUntil(double InPlayhead, TArray<FLiveLinkR
 	}
 }
 
-void FLiveLinkPlaybackTrack::GetFramesUntilReverse(double InPlayhead, TArray<FLiveLinkRecordedFrame>& OutFrames)
+void FLiveLinkPlaybackTrack::GetFramesUntilReverse(const FQualifiedFrameTime& InFrameTime, TArray<FLiveLinkRecordedFrame>& OutFrames)
 {
 	if (LastReadRelativeIndex == INDEX_NONE)
 	{
@@ -148,7 +149,9 @@ void FLiveLinkPlaybackTrack::GetFramesUntilReverse(double InPlayhead, TArray<FLi
 	// We need to look up what the last frame would be if this was running forward, and then end on that frame.
 	// Since we iterate in reverse, but all other operations like GoToFrame use forward look ahead, it's possible the time stamp comparison
 	// will differ by a frame with a reverse look up. There's probably a better way of handling this.
-	const int32 FinalFrameIndex = PlayheadToFrameIndex(InPlayhead);
+	const int32 FinalFrameIndex = ConvertFrameTimeToFrameIndex(InFrameTime);
+
+	const double InTimeSeconds = InFrameTime.AsSeconds();
 	
 	for (FLiveLinkPlaybackTrackReverseIterator It = FLiveLinkPlaybackTrackReverseIterator(*this, GetRelativeIndex(LastReadAbsoluteIndex)); It; ++It)
 	{
@@ -161,7 +164,7 @@ void FLiveLinkPlaybackTrack::GetFramesUntilReverse(double InPlayhead, TArray<FLi
 			continue;
 		}
 		
-		if (FrameTimestamp < InPlayhead || FinalFrameIndex == LastReadRelativeIndex + StartIndexOffset)
+		if (FrameTimestamp < InTimeSeconds || FinalFrameIndex == LastReadRelativeIndex + StartIndexOffset)
 		{
 			break;
 		}
@@ -182,20 +185,7 @@ void FLiveLinkPlaybackTrack::GetFramesUntilReverse(double InPlayhead, TArray<FLi
 
 bool FLiveLinkPlaybackTrack::TryGetFrame(const FQualifiedFrameTime& InFrameTime, FLiveLinkRecordedFrame& OutFrame)
 {
-	int32 LocalizedIndex;
-	if (LocalFrameRate.IsValid())
-	{
-		// Localize the frame time relative to this track's framerate.
-		const FFrameTime LocalizedFrameTime = InFrameTime.ConvertTo(LocalFrameRate);
-		LocalizedIndex = LocalizedFrameTime.GetFrame().Value;
-	}
-	else
-	{
-		LocalizedIndex = InFrameTime.Time.GetFrame().Value;
-	}
-	
-	const int32 RelativeIndex = GetRelativeIndex(LocalizedIndex);
-	
+	const int32 RelativeIndex = ConvertFrameTimeToFrameIndex(InFrameTime);
 	if (RelativeIndex >= 0 && RelativeIndex < FrameData.Num())
 	{
 		LastReadRelativeIndex = RelativeIndex;
@@ -216,35 +206,55 @@ bool FLiveLinkPlaybackTrack::TryGetFrame(const FQualifiedFrameTime& InFrameTime,
 	return false;
 }
 
-int32 FLiveLinkPlaybackTrack::PlayheadToFrameIndex(double InPlayhead)
+int32 FLiveLinkPlaybackTrack::ConvertFrameTimeToFrameIndex(const FQualifiedFrameTime& InFrameTime)
 {
-	int32 CurrentIndex = 0;
+	const double InTimeSeconds = InFrameTime.AsSeconds();
+	
+	// Start by localizing the index based on the framerate. It's possible this is either exact, or closer to the desired position.
+	int32 LocalizedIndex;
+	if (LocalFrameRate.IsValid())
+	{
+		// Localize the frame time relative to this track's framerate.
+		const FFrameTime LocalizedFrameTime = InFrameTime.ConvertTo(LocalFrameRate);
+		LocalizedIndex = LocalizedFrameTime.GetFrame().Value;
+	}
+	else
+	{
+		LocalizedIndex = InFrameTime.Time.GetFrame().Value;
+	}
 
-	for (int32 Idx = 0; Idx < Timestamps.Num(); ++Idx)
+	int32 CurrentIndex = GetRelativeIndex(LocalizedIndex);
+
+	// Now find the closest timestamp without going over. The calculated CurrentIndex should be close to the desired location.
+	if (Timestamps.IsValidIndex(CurrentIndex))
 	{
-		if (Timestamps[Idx] > InPlayhead)
+		// Check if we're ahead of the desired time, iterate backwards if needed.
+		while (CurrentIndex > 0 && Timestamps[CurrentIndex] > InTimeSeconds)
 		{
-			break;
+			--CurrentIndex;
 		}
-
-		CurrentIndex = Idx;
+	}
+	else
+	{
+		// Fall back to full linear search.
+		CurrentIndex = 0;
 	}
 
-	return CurrentIndex + StartIndexOffset;
-}
-
-double FLiveLinkPlaybackTrack::FrameIndexToPlayhead(int32 InIndex)
-{
-	InIndex = GetRelativeIndex(InIndex);
-	if (InIndex >= 0 && InIndex < Timestamps.Num())
+	// If we're behind the desired time, iterate forward as needed.
+	for (int32 Idx = CurrentIndex; Idx < Timestamps.Num(); ++Idx)
 	{
-		return Timestamps[InIndex];
+		if (Timestamps[Idx] > InTimeSeconds)
+		{
+			break;
+		}
+
+		CurrentIndex = Idx;
 	}
 
-	return INDEX_NONE;
+	return CurrentIndex;
 }
 
-TArray<FLiveLinkRecordedFrame> FLiveLinkPlaybackTracks::FetchNextFrames(double Playhead)
+TArray<FLiveLinkRecordedFrame> FLiveLinkPlaybackTracks::FetchNextFrames(const FQualifiedFrameTime& InFrameTime)
 {
 	TArray<FLiveLinkRecordedFrame> NextFrames;
 
@@ -254,14 +264,14 @@ TArray<FLiveLinkRecordedFrame> FLiveLinkPlaybackTracks::FetchNextFrames(double P
 		for (TTuple<FLiveLinkSubjectKey, FLiveLinkPlaybackTrack>& TrackKeyVal : Tracks)
 		{
 			FLiveLinkPlaybackTrack& Track = TrackKeyVal.Value;
-			Track.GetFramesUntil(Playhead, NextFrames);
+			Track.GetFramesUntil(InFrameTime, NextFrames);
 		}
 	}
 
 	return NextFrames;
 }
 
-TArray<FLiveLinkRecordedFrame> FLiveLinkPlaybackTracks::FetchPreviousFrames(double Playhead)
+TArray<FLiveLinkRecordedFrame> FLiveLinkPlaybackTracks::FetchPreviousFrames(const FQualifiedFrameTime& InFrameTime)
 {
 	TArray<FLiveLinkRecordedFrame> PreviousFrames;
 
@@ -271,7 +281,7 @@ TArray<FLiveLinkRecordedFrame> FLiveLinkPlaybackTracks::FetchPreviousFrames(doub
 		for (TTuple<FLiveLinkSubjectKey, FLiveLinkPlaybackTrack>& TrackKeyVal : Tracks)
 		{
 			FLiveLinkPlaybackTrack& Track = TrackKeyVal.Value;
-			Track.GetFramesUntilReverse(Playhead, PreviousFrames);
+			Track.GetFramesUntilReverse(InFrameTime, PreviousFrames);
 		}
 	}
 
@@ -295,36 +305,6 @@ TArray<FLiveLinkRecordedFrame> FLiveLinkPlaybackTracks::FetchNextFramesAtIndex(c
 	return NextFrames;
 }
 
-int32 FLiveLinkPlaybackTracks::PlayheadToFrameIndex(double InPlayhead)
-{
-	for (TTuple<FLiveLinkSubjectKey, FLiveLinkPlaybackTrack>& TrackKeyVal : Tracks)
-	{
-		FLiveLinkPlaybackTrack& Track = TrackKeyVal.Value;
-		// todo: Is this the best way to determine if this is keyframe data and not static data?
-		if (Track.LiveLinkRole == nullptr)
-		{
-			return Track.PlayheadToFrameIndex(InPlayhead);
-		}
-	}
-
-	return INDEX_NONE;
-}
-
-double FLiveLinkPlaybackTracks::FrameIndexToPlayhead(int32 InIndex)
-{
-	for (TTuple<FLiveLinkSubjectKey, FLiveLinkPlaybackTrack>& TrackKeyVal : Tracks)
-	{
-		FLiveLinkPlaybackTrack& Track = TrackKeyVal.Value;
-		// todo: Is this the best way to determine if this is keyframe data and not static data?
-		if (Track.LiveLinkRole == nullptr)
-		{
-			return Track.FrameIndexToPlayhead(InIndex);
-		}
-	}
-
-	return INDEX_NONE;
-}
-
 void FLiveLinkPlaybackTracks::Restart(int32 InIndex)
 {
 	for (TTuple<FLiveLinkSubjectKey, FLiveLinkPlaybackTrack>& TrackKeyVal : Tracks)
diff --git a/Engine/Plugins/Animation/LiveLinkHub/Source/LiveLinkHub/Private/Recording/Implementations/LiveLinkUAssetRecordingPlayer.h b/Engine/Plugins/Animation/LiveLinkHub/Source/LiveLinkHub/Private/Recording/Implementations/LiveLinkUAssetRecordingPlayer.h
--- a/Engine/Plugins/Animation/LiveLinkHub/Source/LiveLinkHub/Private/Recording/Implementations/LiveLinkUAssetRecordingPlayer.h
+++ b/Engine/Plugins/Animation/LiveLinkHub/Source/LiveLinkHub/Private/Recording/Implementations/LiveLinkUAssetRecordingPlayer.h
@@ -9,27 +9,21 @@
 /** Playback track that holds recorded data for a given subject. */
 struct FLiveLinkPlaybackTrack
 {
-	/** Retrieve all frames from last read index to the new playhead, forward looking. */
-	void GetFramesUntil(double InPlayhead, TArray<FLiveLinkRecordedFrame>& OutFrames);
+	/** Retrieve all frames from the last read index to the new frame time, forward-looking. */
+	void GetFramesUntil(const FQualifiedFrameTime& InFrameTime, TArray<FLiveLinkRecordedFrame>& OutFrames);
 
-	/** Retrieve all frames from last read index to the new playhead, reverse looking. */
-	void GetFramesUntilReverse(double InPlayhead, TArray<FLiveLinkRecordedFrame>& OutFrames);
+	/** Retrieve all frames from the last read index to the new frame time, reverse-looking. */
+	void GetFramesUntilReverse(const FQualifiedFrameTime& InFrameTime, TArray<FLiveLinkRecordedFrame>& OutFrames);
 	
 	/** Retrieve the frame at the read index. */
 	bool TryGetFrame(const FQualifiedFrameTime& InFrameTime, FLiveLinkRecordedFrame& OutFrame);
 
 	/**
-	 * Convert the playhead time to a frame index.
-	 * @param InPlayhead The play head to convert to a frame index.
-	 * @return The frame index or INDEX_NONE.
+	 * Given a frame time, find the closest index with a matching timestamp without going over.
+	 * @param InFrameTime The frame time to convert to a frame index.
+	 * @return The frame index or INDEX_NONE
 	 */
-	int32 PlayheadToFrameIndex(double InPlayhead);
-
-	/**
-	 * Convert the frameindex to a playhead.
-	 * @return The frame index or INDEX_NONE.
-	 */
-	double FrameIndexToPlayhead(int32 InIndex);
+	int32 ConvertFrameTimeToFrameIndex(const FQualifiedFrameTime& InFrameTime);
 
 	/** Reset the LastReadIndex. */
 	void Restart(int32 NewIndex = INDEX_NONE)
@@ -74,20 +68,15 @@ private:
 struct FLiveLinkPlaybackTracks
 {
 	/** Get the next frames */
-	TArray<FLiveLinkRecordedFrame> FetchNextFrames(double Playhead);
+	TArray<FLiveLinkRecordedFrame> FetchNextFrames(const FQualifiedFrameTime& InFrameTime);
 
 	/** Get the previous frames as if going in reverse */
-	TArray<FLiveLinkRecordedFrame> FetchPreviousFrames(double Playhead);
+	TArray<FLiveLinkRecordedFrame> FetchPreviousFrames(const FQualifiedFrameTime& InFrameTime);
 	
 	/** Get the next frame(s) at the index */
 	TArray<FLiveLinkRecordedFrame> FetchNextFramesAtIndex(const FQualifiedFrameTime& InFrameTime);
 
-	/** Convert the playhead to a frame index */
-	int32 PlayheadToFrameIndex(double InPlayhead);
-
-	/** Convert the index to a playhead */
-	double FrameIndexToPlayhead(int32 InIndex);
-	
+	/** Restart all tracks. */
 	void Restart(int32 InIndex);
 
 	/** Retrieve the framerate of the first frame */
@@ -109,7 +98,7 @@ public:
 	{
 		if (StreamPlayback(InFrameTime.Time.GetFrame().Value))
 		{
-			return CurrentRecordingPlayback.FetchNextFrames(InFrameTime.AsSeconds());
+			return CurrentRecordingPlayback.FetchNextFrames(InFrameTime);
 		}
 		return TArray<FLiveLinkRecordedFrame>();
 	}
@@ -118,7 +107,7 @@ public:
 	{
 		if (StreamPlayback(InFrameTime.Time.GetFrame().Value))
 		{
-			return CurrentRecordingPlayback.FetchPreviousFrames(InFrameTime.AsSeconds());
+			return CurrentRecordingPlayback.FetchPreviousFrames(InFrameTime);
 		}
 		return TArray<FLiveLinkRecordedFrame>();
 	}
