Add OutboundName to LiveLinkHub Animation Virtual subject

diff --git a/Engine/Plugins/Animation/LiveLink/Source/LiveLink/Private/LiveLinkClient.cpp b/Engine/Plugins/Animation/LiveLink/Source/LiveLink/Private/LiveLinkClient.cpp
--- a/Engine/Plugins/Animation/LiveLink/Source/LiveLink/Private/LiveLinkClient.cpp
+++ b/Engine/Plugins/Animation/LiveLink/Source/LiveLink/Private/LiveLinkClient.cpp
@@ -1436,7 +1436,7 @@ FText FLiveLinkClient::GetSubjectDisplayName(const FLiveLinkSubjectKey& InSubjec
 		}
 		else if (ULiveLinkVirtualSubject* VirtualSubject = Cast<ULiveLinkVirtualSubject>(Settings))
 		{
-			VirtualSubject->GetDisplayName();
+			DisplayName = VirtualSubject->GetDisplayName();
 		}
 	}
 
@@ -1649,13 +1649,21 @@ FLiveLinkSubjectTimeSyncData FLiveLinkClient::GetTimeSyncData(FLiveLinkSubjectNa
 
 FName FLiveLinkClient::GetRebroadcastName(const FLiveLinkSubjectKey& InSubjectKey) const
 {
-	if (ULiveLinkSubjectSettings* Settings = Cast<ULiveLinkSubjectSettings>(GetSubjectSettings(InSubjectKey)))
+	FName RebroadcastName = InSubjectKey.SubjectName.Name;
+
+	if (UObject* Settings = GetSubjectSettings(InSubjectKey))
 	{
-		FName RebroadcastName = Settings->GetRebroadcastName();
-		return RebroadcastName.IsNone() ? InSubjectKey.SubjectName.Name : RebroadcastName;
+		if (ULiveLinkSubjectSettings* SubjectSettings = Cast<ULiveLinkSubjectSettings>(Settings))
+		{
+			RebroadcastName = SubjectSettings->GetRebroadcastName();
+		}
+		else if (ULiveLinkVirtualSubject* VSubject = Cast<ULiveLinkVirtualSubject>(Settings))
+		{
+			RebroadcastName = VSubject->GetRebroadcastName();
+		}
 	}
 
-	return InSubjectKey.SubjectName;
+	return RebroadcastName;
 }
 
 void FLiveLinkClient::PushPendingSubject_AnyThread(FPendingSubjectStatic&& PendingSubject)
diff --git a/Engine/Plugins/Animation/LiveLink/Source/LiveLinkEditor/Private/LiveLinkClientPanelToolbar.cpp b/Engine/Plugins/Animation/LiveLink/Source/LiveLinkEditor/Private/LiveLinkClientPanelToolbar.cpp
--- a/Engine/Plugins/Animation/LiveLink/Source/LiveLinkEditor/Private/LiveLinkClientPanelToolbar.cpp
+++ b/Engine/Plugins/Animation/LiveLink/Source/LiveLinkEditor/Private/LiveLinkClientPanelToolbar.cpp
@@ -2,6 +2,7 @@
 
 #include "LiveLinkClientPanelToolbar.h"
 
+#include "Algo/Accumulate.h"
 #include "Algo/StableSort.h"
 #include "AssetRegistry/AssetRegistryModule.h"
 #include "AssetToolsModule.h"
@@ -64,13 +65,25 @@ public:
 	void Construct(const FArguments& InArgs)
 	{
 		static const FName DefaultVirtualSubjectName = TEXT("Virtual");
+
 		bOkClicked = false;
 		VirtualSubjectClass = nullptr;
-		VirtualSubjectName = DefaultVirtualSubjectName;
 		LiveLinkClient = InArgs._LiveLinkClient;
 
 		check(LiveLinkClient);
 
+		int32 NumVirtualSubjects = Algo::TransformAccumulate(LiveLinkClient->GetSubjects(true, true), [this](const FLiveLinkSubjectKey& SubjectKey)
+			{
+				return LiveLinkClient->IsVirtualSubject(SubjectKey) ? 1 : 0;
+			}, 0);
+
+		VirtualSubjectName = DefaultVirtualSubjectName;
+
+		if (NumVirtualSubjects > 0)
+		{
+			VirtualSubjectName = *FString::Printf(TEXT("%s %d"), *DefaultVirtualSubjectName.ToString(), NumVirtualSubjects + 1);
+		}
+
 		//Default VirtualSubject Source should always exist
 		TArray<FGuid> Sources = LiveLinkClient->GetVirtualSources();
 		check(Sources.Num() > 0);
diff --git a/Engine/Plugins/Animation/LiveLinkHub/Source/LiveLinkHub/Private/LiveLinkHubAnimationVirtualSubject.h b/Engine/Plugins/Animation/LiveLinkHub/Source/LiveLinkHub/Private/LiveLinkHubAnimationVirtualSubject.h
--- a/Engine/Plugins/Animation/LiveLinkHub/Source/LiveLinkHub/Private/LiveLinkHubAnimationVirtualSubject.h
+++ b/Engine/Plugins/Animation/LiveLinkHub/Source/LiveLinkHub/Private/LiveLinkHubAnimationVirtualSubject.h
@@ -6,6 +6,7 @@
 
 #include "Clients/LiveLinkHubProvider.h"
 #include "Features/IModularFeatures.h"
+#include "ILiveLinkClient.h"
 #include "LiveLinkClient.h"
 #include "LiveLinkHubModule.h"
 
@@ -24,12 +25,33 @@ public:
 	virtual void Initialize(FLiveLinkSubjectKey InSubjectKey, TSubclassOf<ULiveLinkRole> InRole, ILiveLinkClient* InLiveLinkClient) override
 	{
 		ULiveLinkAnimationVirtualSubject::Initialize(InSubjectKey, InRole, InLiveLinkClient);
-		SubjectName = InSubjectKey.SubjectName.ToString();
+		OutboundName = InSubjectKey.SubjectName.ToString();
 
 		FLiveLinkClient* Client = &IModularFeatures::Get().GetModularFeature<FLiveLinkClient>(FLiveLinkClient::ModularFeatureName);
 		Source = Client->GetSourceType(InSubjectKey.Source).ToString();
 	}
 
+	virtual FText GetDisplayName() const override
+	{
+		FText DisplayName;
+
+		if (*OutboundName == SubjectKey.SubjectName)
+		{
+			DisplayName = FText::FromName(SubjectKey.SubjectName);
+		}
+		else
+		{
+			DisplayName = FText::Format(INVTEXT("{0} ({1})"), FText::FromString(OutboundName), FText::FromName(SubjectKey.SubjectName));
+		}
+
+		return DisplayName;
+	}
+
+	virtual FName GetRebroadcastName() const override
+	{
+		return *OutboundName;
+	}
+
 	/** Whether this subject is rebroadcasted */
 	virtual bool IsRebroadcasted() const
 	{
@@ -49,17 +71,95 @@ public:
 				FLiveLinkStaticDataStruct StaticDataCopy;
 				StaticDataCopy.InitializeWith(CurrentSnapshot.StaticData);
 
-				LiveLinkProvider->UpdateSubjectStaticData(SubjectKey.SubjectName, Role, MoveTemp(StaticDataCopy));
+				LiveLinkProvider->UpdateSubjectStaticData(GetRebroadcastName(), Role, MoveTemp(StaticDataCopy));
+			}
+		}
+	}
+
+	//~ Begin UObject interface
+	virtual void PreEditChange(FProperty* PropertyAboutToChange) override
+	{
+		Super::PreEditChange(PropertyAboutToChange);
+
+		if (PropertyAboutToChange && PropertyAboutToChange->GetFName() == GET_MEMBER_NAME_CHECKED(ULiveLinkHubAnimationVirtualSubject, OutboundName))
+		{
+			PreviousOutboundName = *OutboundName;
+
+			FLiveLinkHubModule& LiveLinkHubModule = FModuleManager::Get().GetModuleChecked<FLiveLinkHubModule>("LiveLinkHub");
+			if (TSharedPtr<FLiveLinkHubProvider> Provider = LiveLinkHubModule.GetLiveLinkProvider())
+			{
+				Provider->SendClearSubjectToConnections(PreviousOutboundName);
 			}
 		}
 	}
 
+	virtual void PostEditChangeProperty(FPropertyChangedEvent& PropertyChangedEvent) override
+	{
+		Super::PostEditChangeProperty(PropertyChangedEvent);
+
+		if (PropertyChangedEvent.GetPropertyName() == GET_MEMBER_NAME_CHECKED(ULiveLinkHubAnimationVirtualSubject, OutboundName))
+		{
+			if (PreviousOutboundName != *OutboundName)
+			{
+				if (!ValidateOutboundName(OutboundName))
+				{
+					OutboundName = PreviousOutboundName.ToString();
+				}
+				else
+				{
+					NotifyRename();
+				}
+			}
+		}
+	}
+	//~ End UObject interface
+
+	bool ValidateOutboundName(const FString& InOutboundNameCandidate) const
+	{
+		if (InOutboundNameCandidate.IsEmpty() || FName(InOutboundNameCandidate) == NAME_None)
+		{
+			return false;
+		}
+
+		if (InOutboundNameCandidate == OutboundName)
+		{
+			return true;
+		}
+
+		// Can't rename to an existing subject.
+		return !LiveLinkClient->IsSubjectValid(FLiveLinkSubjectName(*InOutboundNameCandidate));
+	}
+
+	void NotifyRename()
+	{
+		FLiveLinkHubModule& LiveLinkHubModule = FModuleManager::Get().GetModuleChecked<FLiveLinkHubModule>("LiveLinkHub");
+
+		if (TSharedPtr<FLiveLinkHubProvider> Provider = LiveLinkHubModule.GetLiveLinkProvider())
+		{
+			// Re-send the last static data with the new name.
+			TPair<UClass*, FLiveLinkStaticDataStruct*> StaticData = Provider->GetLastSubjectStaticDataStruct(PreviousOutboundName);
+			if (StaticData.Key && StaticData.Value)
+			{
+				FLiveLinkStaticDataStruct StaticDataCopy;
+				StaticDataCopy.InitializeWith(*StaticData.Value);
+
+				Provider->UpdateSubjectStaticData(*OutboundName, StaticData.Key, MoveTemp(StaticDataCopy));
+			}
+
+			// Then clear the old static data entry in the provider.
+			Provider->RemoveSubject(PreviousOutboundName);
+		}
+	}
+
 public:
 	/** Name of the virtual subject. */
-	UPROPERTY(VisibleAnywhere, Category = "LiveLink")
-	FString SubjectName;
+	UPROPERTY(EditAnywhere, Category = "LiveLink")
+	FString OutboundName;
 
 	/** Source that contains the subject. */
 	UPROPERTY(VisibleAnywhere, Category = "LiveLink")
 	FString Source;
+
+	/* Previous outbound name to be used for reverting name changes that aren't valid. */
+	FName PreviousOutboundName;
 };
diff --git a/Engine/Plugins/Animation/LiveLinkHub/Source/LiveLinkHub/Private/LiveLinkHubSubjectSettings.cpp b/Engine/Plugins/Animation/LiveLinkHub/Source/LiveLinkHub/Private/LiveLinkHubSubjectSettings.cpp
--- a/Engine/Plugins/Animation/LiveLinkHub/Source/LiveLinkHub/Private/LiveLinkHubSubjectSettings.cpp
+++ b/Engine/Plugins/Animation/LiveLinkHub/Source/LiveLinkHub/Private/LiveLinkHubSubjectSettings.cpp
@@ -67,7 +67,10 @@ void ULiveLinkHubSubjectSettings::PostEditChangeProperty(FPropertyChangedEvent&
 			{
 				OutboundName = PreviousOutboundName.ToString();
 			}
-			NotifyRename();
+			else
+			{
+				NotifyRename();
+			}
 		}
 	}
 	else if (PropertyChangedEvent.GetPropertyName() == GET_MEMBER_NAME_CHECKED(ULiveLinkHubSubjectSettings, Translators)
@@ -104,7 +107,7 @@ void ULiveLinkHubSubjectSettings::PostEditChangeProperty(FPropertyChangedEvent&
 
 bool ULiveLinkHubSubjectSettings::ValidateOutboundName(const FString& InOutboundNameCandidate) const
 {
-	if (InOutboundNameCandidate.IsEmpty())
+	if (InOutboundNameCandidate.IsEmpty() || FName(InOutboundNameCandidate) == NAME_None)
 	{
 		return false;
 	}
diff --git a/Engine/Source/Runtime/LiveLinkInterface/Public/LiveLinkVirtualSubject.h b/Engine/Source/Runtime/LiveLinkInterface/Public/LiveLinkVirtualSubject.h
--- a/Engine/Source/Runtime/LiveLinkInterface/Public/LiveLinkVirtualSubject.h
+++ b/Engine/Source/Runtime/LiveLinkInterface/Public/LiveLinkVirtualSubject.h
@@ -82,6 +82,12 @@ public:
 	{
 		return FText::FromName(SubjectKey.SubjectName);
 	}
+
+	/** Get the name that should be used when the subject is rebroadcast. */
+	virtual FName GetRebroadcastName() const
+	{
+		return SubjectKey.SubjectName;
+	}
 	
 protected:
 
