LiveLinkHub - Fix rename to same name as other subject breaking the livelink connection

diff --git a/Engine/Plugins/Animation/LiveLinkHub/Source/LiveLinkHub/Private/LiveLinkHubAnimationVirtualSubject.h b/Engine/Plugins/Animation/LiveLinkHub/Source/LiveLinkHub/Private/LiveLinkHubAnimationVirtualSubject.h
--- a/Engine/Plugins/Animation/LiveLinkHub/Source/LiveLinkHub/Private/LiveLinkHubAnimationVirtualSubject.h
+++ b/Engine/Plugins/Animation/LiveLinkHub/Source/LiveLinkHub/Private/LiveLinkHubAnimationVirtualSubject.h
@@ -9,6 +9,7 @@
 #include "ILiveLinkClient.h"
 #include "LiveLinkClient.h"
 #include "LiveLinkHubModule.h"
+#include "LiveLinkHubSubjectSettingsUtils.h"
 
 #include "LiveLinkHubAnimationVirtualSubject.generated.h"
 
@@ -101,56 +102,19 @@ public:
 		{
 			if (PreviousOutboundName != *OutboundName)
 			{
-				if (!ValidateOutboundName(OutboundName))
+				if (!FLiveLinkHubSubjectSettingsUtils::ValidateOutboundName(OutboundName, PreviousOutboundName, OutboundName))
 				{
 					OutboundName = PreviousOutboundName.ToString();
 				}
 				else
 				{
-					NotifyRename();
+					FLiveLinkHubSubjectSettingsUtils::NotifyRename(PreviousOutboundName, OutboundName);
 				}
 			}
 		}
 	}
 	//~ End UObject interface
 
-	bool ValidateOutboundName(const FString& InOutboundNameCandidate) const
-	{
-		if (InOutboundNameCandidate.IsEmpty() || FName(InOutboundNameCandidate) == NAME_None)
-		{
-			return false;
-		}
-
-		if (InOutboundNameCandidate == OutboundName)
-		{
-			return true;
-		}
-
-		// Can't rename to an existing subject.
-		return !LiveLinkClient->IsSubjectValid(FLiveLinkSubjectName(*InOutboundNameCandidate));
-	}
-
-	void NotifyRename()
-	{
-		FLiveLinkHubModule& LiveLinkHubModule = FModuleManager::Get().GetModuleChecked<FLiveLinkHubModule>("LiveLinkHub");
-
-		if (TSharedPtr<FLiveLinkHubProvider> Provider = LiveLinkHubModule.GetLiveLinkProvider())
-		{
-			// Re-send the last static data with the new name.
-			TPair<UClass*, FLiveLinkStaticDataStruct*> StaticData = Provider->GetLastSubjectStaticDataStruct(PreviousOutboundName);
-			if (StaticData.Key && StaticData.Value)
-			{
-				FLiveLinkStaticDataStruct StaticDataCopy;
-				StaticDataCopy.InitializeWith(*StaticData.Value);
-
-				Provider->UpdateSubjectStaticData(*OutboundName, StaticData.Key, MoveTemp(StaticDataCopy));
-			}
-
-			// Then clear the old static data entry in the provider.
-			Provider->RemoveSubject(PreviousOutboundName);
-		}
-	}
-
 public:
 	/** Name of the virtual subject. */
 	UPROPERTY(EditAnywhere, Category = "LiveLink")
diff --git a/Engine/Plugins/Animation/LiveLinkHub/Source/LiveLinkHub/Private/LiveLinkHubSubjectSettings.cpp b/Engine/Plugins/Animation/LiveLinkHub/Source/LiveLinkHub/Private/LiveLinkHubSubjectSettings.cpp
--- a/Engine/Plugins/Animation/LiveLinkHub/Source/LiveLinkHub/Private/LiveLinkHubSubjectSettings.cpp
+++ b/Engine/Plugins/Animation/LiveLinkHub/Source/LiveLinkHub/Private/LiveLinkHubSubjectSettings.cpp
@@ -4,6 +4,7 @@
 
 #include "LiveLinkHubClient.h"
 #include "LiveLinkHubModule.h"
+#include "LiveLinkHubSubjectSettingsUtils.h"
 #include "Features/IModularFeatures.h"
 #include "Clients/LiveLinkHubProvider.h"
 
@@ -18,27 +19,6 @@ void ULiveLinkHubSubjectSettings::Initialize(FLiveLinkSubjectKey InSubjectKey)
 	Source = LiveLinkClient->GetSourceType(InSubjectKey.Source).ToString();
 }
 
-void ULiveLinkHubSubjectSettings::NotifyRename()
-{
-	FLiveLinkHubModule& LiveLinkHubModule = FModuleManager::Get().GetModuleChecked<FLiveLinkHubModule>("LiveLinkHub");
-
-	if (TSharedPtr<FLiveLinkHubProvider> Provider = LiveLinkHubModule.GetLiveLinkProvider())
-	{
-		// Re-send the last static data with the new name.
-		TPair<UClass*, FLiveLinkStaticDataStruct*> StaticData = Provider->GetLastSubjectStaticDataStruct(PreviousOutboundName);
-		if (StaticData.Key && StaticData.Value)
-		{
-			FLiveLinkStaticDataStruct StaticDataCopy;
-			StaticDataCopy.InitializeWith(*StaticData.Value);
-
-			Provider->UpdateSubjectStaticData(*OutboundName, StaticData.Key, MoveTemp(StaticDataCopy));
-		}
-
-		// Then clear the old static data entry in the provider.
-		Provider->RemoveSubject(PreviousOutboundName);
-	}
-}
-
 void ULiveLinkHubSubjectSettings::PreEditChange(FProperty* PropertyAboutToChange)
 {
 	Super::PreEditChange(PropertyAboutToChange);
@@ -63,13 +43,13 @@ void ULiveLinkHubSubjectSettings::PostEditChangeProperty(FPropertyChangedEvent&
 	{
 		if (PreviousOutboundName != *OutboundName)
 		{
-			if (!ValidateOutboundName(OutboundName))
+			if (!FLiveLinkHubSubjectSettingsUtils::ValidateOutboundName(SubjectName, PreviousOutboundName, OutboundName))
 			{
 				OutboundName = PreviousOutboundName.ToString();
 			}
 			else
 			{
-				NotifyRename();
+				FLiveLinkHubSubjectSettingsUtils::NotifyRename(PreviousOutboundName, OutboundName);
 			}
 		}
 	}
@@ -104,21 +84,3 @@ void ULiveLinkHubSubjectSettings::PostEditChangeProperty(FPropertyChangedEvent&
 		LiveLinkClient->CacheSubjectSettings(Key, this);
 	}
 }
-
-bool ULiveLinkHubSubjectSettings::ValidateOutboundName(const FString& InOutboundNameCandidate) const
-{
-	if (InOutboundNameCandidate.IsEmpty() || FName(InOutboundNameCandidate) == NAME_None)
-	{
-		return false;
-	}
-
-	if (InOutboundNameCandidate == SubjectName)
-	{
-		return true;
-	}
-
-	FLiveLinkHubClient* LiveLinkClient = static_cast<FLiveLinkHubClient*>(&IModularFeatures::Get().GetModularFeature<ILiveLinkClient>(ILiveLinkClient::ModularFeatureName));
-
-	// Can't rename to an existing subject.
-	return !LiveLinkClient->IsSubjectValid(FLiveLinkSubjectName(*InOutboundNameCandidate));
-}
diff --git a/Engine/Plugins/Animation/LiveLinkHub/Source/LiveLinkHub/Private/LiveLinkHubSubjectSettings.h b/Engine/Plugins/Animation/LiveLinkHub/Source/LiveLinkHub/Private/LiveLinkHubSubjectSettings.h
--- a/Engine/Plugins/Animation/LiveLinkHub/Source/LiveLinkHub/Private/LiveLinkHubSubjectSettings.h
+++ b/Engine/Plugins/Animation/LiveLinkHub/Source/LiveLinkHub/Private/LiveLinkHubSubjectSettings.h
@@ -25,18 +25,11 @@ public:
 	}
 	//~ End ULiveLinkSubjectSettings interface
 
-	/** Notify connected clients that this subject's name has changed. */
-	void NotifyRename();
-
 	//~ Begin UObject interface
 	virtual void PreEditChange(FProperty* PropertyAboutToChange) override;
 	virtual void PostEditChangeProperty(FPropertyChangedEvent& PropertyChangedEvent) override;
 	//~ End UObject interface
 
-private:
-	/** Returns whether a new name candidate for the outbound name is valid. */
-	bool ValidateOutboundName(const FString& InOutboundNameCandidate) const;
-
 public:
 	/** Name of this subject. */
 	UPROPERTY(VisibleAnywhere, Category = "LiveLink")
diff --git a/Engine/Plugins/Animation/LiveLinkHub/Source/LiveLinkHub/Private/LiveLinkHubSubjectSettingsUtils.h b/Engine/Plugins/Animation/LiveLinkHub/Source/LiveLinkHub/Private/LiveLinkHubSubjectSettingsUtils.h
new file mode 100644
--- /dev/null
+++ b/Engine/Plugins/Animation/LiveLinkHub/Source/LiveLinkHub/Private/LiveLinkHubSubjectSettingsUtils.h
@@ -0,0 +1,55 @@
+// Copyright Epic Games, Inc. All Rights Reserved.
+
+#pragma once
+
+#include "Clients/LiveLinkHubProvider.h"
+#include "Features/IModularFeatures.h"
+#include "LiveLinkHubModule.h"
+#include "LiveLinkHubClient.h"
+#include "LiveLinkTypes.h"
+
+class FLiveLinkHubSubjectSettingsUtils
+{
+public:
+	/** Notify connected clients that this subject's name has changed. */
+	static bool ValidateOutboundName(const FString& SubjectName, FName PreviousOutboundName, const FString& InOutboundNameCandidate)
+	{
+		if (InOutboundNameCandidate.IsEmpty() || FName(InOutboundNameCandidate) == NAME_None)
+		{
+			return false;
+		}
+
+		if (InOutboundNameCandidate == SubjectName)
+		{
+			return true;
+		}
+
+		FLiveLinkHubClient* LiveLinkClient = static_cast<FLiveLinkHubClient*>(&IModularFeatures::Get().GetModularFeature<ILiveLinkClient>(ILiveLinkClient::ModularFeatureName));
+
+		// Can't rename to an existing subject, so check if it exists by querying for its role.
+		return LiveLinkClient->GetSubjectRole_AnyThread(FLiveLinkSubjectName(*InOutboundNameCandidate)).Get() == nullptr;
+	}
+
+	/** Returns whether a new name candidate for the outbound name is valid. */
+	static void NotifyRename(FName PreviousOutboundName, const FString& OutboundName)
+	{
+		FLiveLinkHubModule& LiveLinkHubModule = FModuleManager::Get().GetModuleChecked<FLiveLinkHubModule>("LiveLinkHub");
+
+		if (TSharedPtr<FLiveLinkHubProvider> Provider = LiveLinkHubModule.GetLiveLinkProvider())
+		{
+			// Re-send the last static data with the new name.
+			TPair<UClass*, FLiveLinkStaticDataStruct*> StaticData = Provider->GetLastSubjectStaticDataStruct(PreviousOutboundName);
+			if (StaticData.Key && StaticData.Value)
+			{
+				FLiveLinkStaticDataStruct StaticDataCopy;
+				StaticDataCopy.InitializeWith(*StaticData.Value);
+
+				Provider->UpdateSubjectStaticData(*OutboundName, StaticData.Key, MoveTemp(StaticDataCopy));
+			}
+
+			// Then clear the old static data entry in the provider.
+			Provider->RemoveSubject(PreviousOutboundName);
+		}
+	}
+
+};
\ No newline at end of file
