Rename LiveLink to Live Link in LiveLinkHub

diff --git a/Engine/Plugins/Animation/LiveLinkHub/Source/LiveLinkHub/Private/Recording/LiveLinkHubPlaybackSourceFactory.h b/Engine/Plugins/Animation/LiveLinkHub/Source/LiveLinkHub/Private/Recording/LiveLinkHubPlaybackSourceFactory.h
--- a/Engine/Plugins/Animation/LiveLinkHub/Source/LiveLinkHub/Private/Recording/LiveLinkHubPlaybackSourceFactory.h
+++ b/Engine/Plugins/Animation/LiveLinkHub/Source/LiveLinkHub/Private/Recording/LiveLinkHubPlaybackSourceFactory.h
@@ -25,7 +25,7 @@ public:
 
 	virtual FText GetSourceTooltip() const
 	{
-		return LOCTEXT("PlaybackSourceTooltip", "Mock livelink source used to playback livelink data that was recorded.");
+		return LOCTEXT("PlaybackSourceTooltip", "Mock Live Link source used to playback Live Link data that was recorded.");
 	}
 
 	virtual TSharedPtr<ILiveLinkSource> CreateSource(const FString& ConnectionString) const
diff --git a/Engine/Plugins/Animation/LiveLinkHub/Source/LiveLinkHubEditor/Private/LiveLinkHubEditorModule.cpp b/Engine/Plugins/Animation/LiveLinkHub/Source/LiveLinkHubEditor/Private/LiveLinkHubEditorModule.cpp
--- a/Engine/Plugins/Animation/LiveLinkHub/Source/LiveLinkHubEditor/Private/LiveLinkHubEditorModule.cpp
+++ b/Engine/Plugins/Animation/LiveLinkHub/Source/LiveLinkHubEditor/Private/LiveLinkHubEditorModule.cpp
@@ -55,8 +55,8 @@ void FLiveLinkHubEditorModule::OnPostEngineInit()
 		FToolMenuSection& Section = Menu->AddSection("VirtualProductionSection", LOCTEXT("VirtualProductionSection", "Virtual Production"));
 
 		Section.AddMenuEntry("LiveLinkHub",
-			LOCTEXT("LiveLinkHubLabel", "LiveLink Hub"),
-			LOCTEXT("LiveLinkHubTooltip", "Launch the LiveLink Hub app."),
+			LOCTEXT("LiveLinkHubLabel", "Live Link Hub"),
+			LOCTEXT("LiveLinkHubTooltip", "Launch the Live Link Hub app."),
 			FSlateIcon("LiveLinkStyle", "LiveLinkClient.Common.Icon.Small"),
 			FUIAction(FExecuteAction::CreateRaw(this, &FLiveLinkHubEditorModule::OpenLiveLinkHub)));
 	}
@@ -66,7 +66,7 @@ void FLiveLinkHubEditorModule::OpenLiveLinkHub()
 {
 	FAsyncTaskNotificationConfig NotificationConfig;
 	NotificationConfig.bKeepOpenOnFailure = true;
-	NotificationConfig.TitleText = LOCTEXT("LaunchingLiveLinkHub", "Launching LiveLink Hub...");
+	NotificationConfig.TitleText = LOCTEXT("LaunchingLiveLinkHub", "Launching Live Link Hub...");
 	NotificationConfig.LogCategory = &LogLiveLinkHubEditor;
 
 	FAsyncTaskNotification Notification(NotificationConfig);
@@ -88,7 +88,7 @@ void FLiveLinkHubEditorModule::OpenLiveLinkHub()
 			{
 				Notification.SetComplete(
 					LaunchLiveLinkHubErrorTitle,
-					LOCTEXT("LaunchLiveLinkHubError_CouldNotOpenLauncher", "Could not launch LiveLinkHub through the Epic Games Store."),
+					LOCTEXT("LaunchLiveLinkHubError_CouldNotOpenLauncher", "Could not launch Live Link Hub through the Epic Games Store."),
 					false
 				);
 
@@ -96,15 +96,15 @@ void FLiveLinkHubEditorModule::OpenLiveLinkHub()
 			else
 			{
 				Notification.SetComplete(
-					LOCTEXT("LiveLinkHubLaunchSuccessTitle", "Launched LiveLinkHub."),
-					LOCTEXT("LaunchLiveLinkHubError_LaunchSuccess", "Launching LiveLinkHub through the Epic Games Store."),
+					LOCTEXT("LiveLinkHubLaunchSuccessTitle", "Launched Live Link Hub."),
+					LOCTEXT("LaunchLiveLinkHubError_LaunchSuccess", "Launching Liv Link Hub through the Epic Games Store."),
 					true
 				);
 			}
 		}
 		else
 		{
-			const EAppReturnType::Type Choice = FMessageDialog::Open(EAppMsgType::YesNo, LOCTEXT("InstallThroughEGS", "LiveLinkHub is not currently installed, do you want to install it through the Epic Games Store?"));
+			const EAppReturnType::Type Choice = FMessageDialog::Open(EAppMsgType::YesNo, LOCTEXT("InstallThroughEGS", "Live Link Hub is not currently installed, do you want to install it through the Epic Games Store?"));
 
 			if (Choice == EAppReturnType::Yes)
 			{
@@ -116,7 +116,7 @@ void FLiveLinkHubEditorModule::OpenLiveLinkHub()
 					{
 						Notification.SetComplete(
 							LaunchLiveLinkHubErrorTitle,
-							LOCTEXT("LaunchLiveLinkHubError_CouldNotFindHubStorePage", "Could not find the LiveLink Hub page on the Epic Games Store."),
+							LOCTEXT("LaunchLiveLinkHubError_CouldNotFindHubStorePage", "Could not find the Live Link Hub page on the Epic Games Store."),
 							false
 						);
 					}
@@ -124,7 +124,7 @@ void FLiveLinkHubEditorModule::OpenLiveLinkHub()
 					{
 						Notification.SetComplete(
 							LaunchLiveLinkHubErrorTitle,
-							LOCTEXT("LaunchLiveLinkHub_LaunchFromStore", "Opening Epic Games Store to the LiveLinkHub page."),
+							LOCTEXT("LaunchLiveLinkHub_LaunchFromStore", "Opening Epic Games Store to the Live Link Hub page."),
 							true
 						);
 					}
@@ -134,7 +134,7 @@ void FLiveLinkHubEditorModule::OpenLiveLinkHub()
 				{
 					Notification.SetComplete(
 						LaunchLiveLinkHubErrorTitle,
-						LOCTEXT("LaunchLiveLinkHubError_EmptyConfig", "Could not find the LiveLink Hub page on the Epic Games Store, missing configuration for the store page."),
+						LOCTEXT("LaunchLiveLinkHubError_EmptyConfig", "Could not find the Live Link Hub page on the Epic Games Store, missing configuration for the store page."),
 						false
 					);
 				
@@ -144,7 +144,7 @@ void FLiveLinkHubEditorModule::OpenLiveLinkHub()
 			{
 				Notification.SetComplete(
 					LaunchLiveLinkHubErrorTitle,
-					LOCTEXT("LaunchLiveLinkHub_DidNotLaunchFromStore", "LiveLinkHub could not be launched since it wasn't installed."),
+					LOCTEXT("LaunchLiveLinkHub_DidNotLaunchFromStore", "Live Link Hub could not be launched since it wasn't installed."),
 					false
 				);
 			}
@@ -173,7 +173,7 @@ void FLiveLinkHubEditorModule::OpenLiveLinkHub()
 	{
 		Notification.SetComplete(
 			LaunchLiveLinkHubErrorTitle,
-			LOCTEXT("LaunchLiveLinkHubError_ExecutableMissing", "Could not find the executable. Have you compiled the LiveLink Hub app?"),
+			LOCTEXT("LaunchLiveLinkHubError_ExecutableMissing", "Could not find the executable. Have you compiled the Live Link Hub app?"),
 			false
 		);
 
@@ -186,7 +186,7 @@ void FLiveLinkHubEditorModule::OpenLiveLinkHub()
 	{
 		Notification.SetComplete(
 			LaunchLiveLinkHubErrorTitle,
-			LOCTEXT("LaunchLiveLinkHubError_AlreadyRunning", "A LiveLinkHub instance is already running."),
+			LOCTEXT("LaunchLiveLinkHubError_AlreadyRunning", "A Live Link Hub instance is already running."),
 			false
 		);
 		return;
@@ -200,7 +200,7 @@ void FLiveLinkHubEditorModule::OpenLiveLinkHub()
 	if (ProcHandle.IsValid())
 	{
 		Notification.SetComplete(
-			LOCTEXT("LaunchedLiveLinkHub", "Launched LiveLink Hub"), FText(), true);
+			LOCTEXT("LaunchedLiveLinkHub", "Launched Live Link Hub"), FText(), true);
 
 		return;
 	}
@@ -208,7 +208,7 @@ void FLiveLinkHubEditorModule::OpenLiveLinkHub()
 	{
 		Notification.SetComplete(
 			LaunchLiveLinkHubErrorTitle,
-			LOCTEXT("LaunchLiveLinkHubError_InvalidHandle", "Failed to create the LiveLink Hub process."),
+			LOCTEXT("LaunchLiveLinkHubError_InvalidHandle", "Failed to create the Live Link Hub process."),
 			false);
 	}
 }
@@ -236,4 +236,4 @@ TSharedRef<SWidget> FLiveLinkHubEditorModule::CreateLiveLinkHubWidget()
 
 IMPLEMENT_MODULE(FLiveLinkHubEditorModule, LiveLinkHubEditor);
 
-#undef LOCTEXT_NAMESPACE /*LiveLinkHubEditor*/ 
\ No newline at end of file
+#undef LOCTEXT_NAMESPACE /*LiveLinkHubEditor*/ 
