Fix subject being removed when renamed to another subject of the same name

diff --git a/Engine/Plugins/Animation/LiveLinkHub/Source/LiveLinkHub/Private/LiveLinkHubAnimationVirtualSubject.h b/Engine/Plugins/Animation/LiveLinkHub/Source/LiveLinkHub/Private/LiveLinkHubAnimationVirtualSubject.h
--- a/Engine/Plugins/Animation/LiveLinkHub/Source/LiveLinkHub/Private/LiveLinkHubAnimationVirtualSubject.h
+++ b/Engine/Plugins/Animation/LiveLinkHub/Source/LiveLinkHub/Private/LiveLinkHubAnimationVirtualSubject.h
@@ -85,12 +85,6 @@ public:
 		if (PropertyAboutToChange && PropertyAboutToChange->GetFName() == GET_MEMBER_NAME_CHECKED(ULiveLinkHubAnimationVirtualSubject, OutboundName))
 		{
 			PreviousOutboundName = *OutboundName;
-
-			FLiveLinkHubModule& LiveLinkHubModule = FModuleManager::Get().GetModuleChecked<FLiveLinkHubModule>("LiveLinkHub");
-			if (TSharedPtr<FLiveLinkHubProvider> Provider = LiveLinkHubModule.GetLiveLinkProvider())
-			{
-				Provider->SendClearSubjectToConnections(PreviousOutboundName);
-			}
 		}
 	}
 
@@ -102,13 +96,19 @@ public:
 		{
 			if (PreviousOutboundName != *OutboundName)
 			{
-				if (!FLiveLinkHubSubjectSettingsUtils::ValidateOutboundName(OutboundName, PreviousOutboundName, OutboundName))
+				if (FLiveLinkHubSubjectSettingsUtils::ValidateOutboundName(OutboundName, PreviousOutboundName, OutboundName))
 				{
-					OutboundName = PreviousOutboundName.ToString();
+					FLiveLinkHubModule& LiveLinkHubModule = FModuleManager::Get().GetModuleChecked<FLiveLinkHubModule>("LiveLinkHub");
+					if (TSharedPtr<FLiveLinkHubProvider> Provider = LiveLinkHubModule.GetLiveLinkProvider())
+					{
+						Provider->SendClearSubjectToConnections(PreviousOutboundName);
+					}
+
+					FLiveLinkHubSubjectSettingsUtils::NotifyRename(PreviousOutboundName, OutboundName, SubjectKey);
 				}
 				else
 				{
-					FLiveLinkHubSubjectSettingsUtils::NotifyRename(PreviousOutboundName, OutboundName);
+					OutboundName = PreviousOutboundName.ToString();
 				}
 			}
 		}
@@ -124,6 +124,7 @@ public:
 	UPROPERTY(VisibleAnywhere, Category = "LiveLink")
 	FString Source;
 
+private:
 	/* Previous outbound name to be used for reverting name changes that aren't valid. */
 	FName PreviousOutboundName;
 };
diff --git a/Engine/Plugins/Animation/LiveLinkHub/Source/LiveLinkHub/Private/LiveLinkHubSubjectSettings.cpp b/Engine/Plugins/Animation/LiveLinkHub/Source/LiveLinkHub/Private/LiveLinkHubSubjectSettings.cpp
--- a/Engine/Plugins/Animation/LiveLinkHub/Source/LiveLinkHub/Private/LiveLinkHubSubjectSettings.cpp
+++ b/Engine/Plugins/Animation/LiveLinkHub/Source/LiveLinkHub/Private/LiveLinkHubSubjectSettings.cpp
@@ -26,12 +26,6 @@ void ULiveLinkHubSubjectSettings::PreEditChange(FProperty* PropertyAboutToChange
 	if (PropertyAboutToChange && PropertyAboutToChange->GetFName() == GET_MEMBER_NAME_CHECKED(ULiveLinkHubSubjectSettings, OutboundName))
 	{
 		PreviousOutboundName = *OutboundName;
-
-		FLiveLinkHubModule& LiveLinkHubModule = FModuleManager::Get().GetModuleChecked<FLiveLinkHubModule>("LiveLinkHub");
-		if (TSharedPtr<FLiveLinkHubProvider> Provider = LiveLinkHubModule.GetLiveLinkProvider())
-		{
-			Provider->SendClearSubjectToConnections(PreviousOutboundName);
-		}
 	}
 }
 
@@ -43,13 +37,19 @@ void ULiveLinkHubSubjectSettings::PostEditChangeProperty(FPropertyChangedEvent&
 	{
 		if (PreviousOutboundName != *OutboundName)
 		{
-			if (!FLiveLinkHubSubjectSettingsUtils::ValidateOutboundName(SubjectName, PreviousOutboundName, OutboundName))
+			if (FLiveLinkHubSubjectSettingsUtils::ValidateOutboundName(SubjectName, PreviousOutboundName, OutboundName))
 			{
-				OutboundName = PreviousOutboundName.ToString();
+				FLiveLinkHubModule& LiveLinkHubModule = FModuleManager::Get().GetModuleChecked<FLiveLinkHubModule>("LiveLinkHub");
+				if (TSharedPtr<FLiveLinkHubProvider> Provider = LiveLinkHubModule.GetLiveLinkProvider())
+				{
+					Provider->SendClearSubjectToConnections(PreviousOutboundName);
+				}
+
+				FLiveLinkHubSubjectSettingsUtils::NotifyRename(PreviousOutboundName, OutboundName, Key);
 			}
 			else
 			{
-				FLiveLinkHubSubjectSettingsUtils::NotifyRename(PreviousOutboundName, OutboundName);
+				OutboundName = PreviousOutboundName.ToString();
 			}
 		}
 	}
diff --git a/Engine/Plugins/Animation/LiveLinkHub/Source/LiveLinkHub/Private/LiveLinkHubSubjectSettingsUtils.h b/Engine/Plugins/Animation/LiveLinkHub/Source/LiveLinkHub/Private/LiveLinkHubSubjectSettingsUtils.h
--- a/Engine/Plugins/Animation/LiveLinkHub/Source/LiveLinkHub/Private/LiveLinkHubSubjectSettingsUtils.h
+++ b/Engine/Plugins/Animation/LiveLinkHub/Source/LiveLinkHub/Private/LiveLinkHubSubjectSettingsUtils.h
@@ -31,7 +31,7 @@ public:
 	}
 
 	/** Returns whether a new name candidate for the outbound name is valid. */
-	static void NotifyRename(FName PreviousOutboundName, const FString& OutboundName)
+	static void NotifyRename(FName PreviousOutboundName, const FString& OutboundName, const FLiveLinkSubjectKey& SubjectKey)
 	{
 		FLiveLinkHubModule& LiveLinkHubModule = FModuleManager::Get().GetModuleChecked<FLiveLinkHubModule>("LiveLinkHub");
 
@@ -44,7 +44,13 @@ public:
 				FLiveLinkStaticDataStruct StaticDataCopy;
 				StaticDataCopy.InitializeWith(*StaticData.Value);
 
-				Provider->UpdateSubjectStaticData(*OutboundName, StaticData.Key, MoveTemp(StaticDataCopy));
+
+				TMap<FName, FString> ExtraAnnotations;
+
+				FLiveLinkHubClient* LiveLinkClient = static_cast<FLiveLinkHubClient*>(&IModularFeatures::Get().GetModularFeature<ILiveLinkClient>(ILiveLinkClient::ModularFeatureName));
+				const FText OriginalSourceType = LiveLinkClient->GetSourceType(SubjectKey.Source);
+				ExtraAnnotations.Add(FLiveLinkMessageAnnotation::OriginalSourceAnnotation, OriginalSourceType.ToString());
+				Provider->UpdateSubjectStaticData(*OutboundName, StaticData.Key, MoveTemp(StaticDataCopy), ExtraAnnotations);
 			}
 
 			// Then clear the old static data entry in the provider.
