Fix LiveLinkHub crash on shutdown caused by calling Broadcast while the engine is shutting down

diff --git a/Engine/Plugins/Animation/LiveLinkHub/Source/LiveLinkHub/Private/Recording/LiveLinkHubPlaybackController.cpp b/Engine/Plugins/Animation/LiveLinkHub/Source/LiveLinkHub/Private/Recording/LiveLinkHubPlaybackController.cpp
--- a/Engine/Plugins/Animation/LiveLinkHub/Source/LiveLinkHub/Private/Recording/LiveLinkHubPlaybackController.cpp
+++ b/Engine/Plugins/Animation/LiveLinkHub/Source/LiveLinkHub/Private/Recording/LiveLinkHubPlaybackController.cpp
@@ -4,6 +4,7 @@
 
 #include "Async/Async.h"
 #include "Async/TaskGraphInterfaces.h"
+#include "CoreGlobals.h"
 #include "ILiveLinkClient.h"
 #include "Implementations/LiveLinkUAssetRecordingPlayer.h"
 #include "LiveLinkHubLog.h"
@@ -538,7 +539,7 @@ uint32 FLiveLinkHubPlaybackController::Run()
 
 void FLiveLinkHubPlaybackController::OnPlaybackFinished_Internal()
 {
-	if (!bIsDestructing) // Can crash otherwise, such as if we are closing the app.
+	if (!bIsDestructing && GIsRunning) // Can crash otherwise, such as if we are closing the app.
 	{
 		PlaybackFinishedDelegate.Broadcast();
 	}
