Fix crash when invoking debug console in LiveLinkHub.

This was caused by GlobalEditorCommonCommands calling into the UStatusBarSubsystem without checking if it's valid first.

[RN] fix Virtual Production/Tools

diff --git a/Engine/Plugins/Animation/LiveLinkHub/Source/LiveLinkHub/Private/LiveLinkHub.cpp b/Engine/Plugins/Animation/LiveLinkHub/Source/LiveLinkHub/Private/LiveLinkHub.cpp
--- a/Engine/Plugins/Animation/LiveLinkHub/Source/LiveLinkHub/Private/LiveLinkHub.cpp
+++ b/Engine/Plugins/Animation/LiveLinkHub/Source/LiveLinkHub/Private/LiveLinkHub.cpp
@@ -58,6 +58,7 @@ void FLiveLinkHub::Initialize(bool bLauncherDistribution)
 	LiveLinkProvider = MakeShared<FLiveLinkHubProvider>(SessionManager.ToSharedRef());
 
 	FModuleManager::Get().LoadModule("Settings");
+	FModuleManager::Get().LoadModule("StatusBar");
 
 	if (bLauncherDistribution)
 	{
@@ -252,9 +253,9 @@ void FLiveLinkHub::UnregisterLiveLinkHubSettings()
 {
 	if (ISettingsModule* SettingsModule = FModuleManager::GetModulePtr<ISettingsModule>("Settings"))
 	{
-		SettingsModule->UnregisterSettings("Editor", "Plugins", "LiveLink");
-		SettingsModule->UnregisterSettings("Project", "Plugins", "LiveLink");
-		SettingsModule->UnregisterSettings("Project", "Plugins", "LiveLinkHub");
+		SettingsModule->UnregisterSettings("Editor", "Plugins", "Live Link");
+		SettingsModule->UnregisterSettings("Project", "Plugins", "Live Link");
+		SettingsModule->UnregisterSettings("Project", "Plugins", "Live Link Hub");
 	}
 }
 
