LiveLinkHub - Fix program freeze when trying to play an empty recording.
Prevent saving over existing recordings to fix a crash on playback.

diff --git a/Engine/Plugins/Animation/LiveLinkHub/Source/LiveLinkHub/Private/Recording/Implementations/LiveLinkUAssetRecorder.cpp b/Engine/Plugins/Animation/LiveLinkHub/Source/LiveLinkHub/Private/Recording/Implementations/LiveLinkUAssetRecorder.cpp
--- a/Engine/Plugins/Animation/LiveLinkHub/Source/LiveLinkHub/Private/Recording/Implementations/LiveLinkUAssetRecorder.cpp
+++ b/Engine/Plugins/Animation/LiveLinkHub/Source/LiveLinkHub/Private/Recording/Implementations/LiveLinkUAssetRecorder.cpp
@@ -203,6 +203,12 @@ bool FLiveLinkUAssetRecorder::GetSavePresetPackageName(FString& OutName)
 
 		NewPackageName = FString::Format(*UserPackageName, FormatArgs);
 
+		{
+			// Ensure we aren't saving over an existing recording as this isn't supported yet... the package isn't saved properly.
+			// todo: We should support saving over/deleting existing recordings. This is simply a low risk fix to a playback crash for 5.5.0.
+			AssetToolsModule.Get().CreateUniqueAssetName(NewPackageName, TEXT(""), NewPackageName, UniqueAssetName);
+		}
+		
 		FText OutError;
 		bFilenameValid = FFileHelper::IsFilenameValidForSaving(NewPackageName, OutError);
 	}
diff --git a/Engine/Plugins/Animation/LiveLinkHub/Source/LiveLinkHub/Private/Recording/Implementations/LiveLinkUAssetRecording.cpp b/Engine/Plugins/Animation/LiveLinkHub/Source/LiveLinkHub/Private/Recording/Implementations/LiveLinkUAssetRecording.cpp
--- a/Engine/Plugins/Animation/LiveLinkHub/Source/LiveLinkHub/Private/Recording/Implementations/LiveLinkUAssetRecording.cpp
+++ b/Engine/Plugins/Animation/LiveLinkHub/Source/LiveLinkHub/Private/Recording/Implementations/LiveLinkUAssetRecording.cpp
@@ -223,9 +223,15 @@ bool ULiveLinkUAssetRecording::WaitForBufferedFrames(int32 InMinFrame, int32 InM
 			FPlatformProcess::Sleep(0.002);
 		}
 
+		const int32 MaxFrames = GetMaxFrames();
+		if (MaxFrames <= 0)
+		{
+			return true;
+		}
+		
 		// Clamp the frame range to the max possible range. If the selection range extends the actual frame range
 		// then there would be nothing to load.
-		const int32 EndFrame = RecordingMaxFrames - 1;
+		const int32 EndFrame = MaxFrames - 1;
 		InMinFrame = FMath::Clamp(InMinFrame, 0, EndFrame);
 		InMaxFrame = FMath::Clamp(InMaxFrame, 0, EndFrame);
 
