Store bIsTimeCodeSource in the livelinkhub session config to fix LLH setting the timecode source on unwanted clients.

Remove cached LiveLinkHub settings from  LiveLinkHub provider and always rely on the session config to be the source of truth for timecode settings.

diff --git a/Engine/Plugins/Animation/LiveLinkHub/Source/LiveLinkHub/Private/Clients/LiveLinkHubProvider.cpp b/Engine/Plugins/Animation/LiveLinkHub/Source/LiveLinkHub/Private/Clients/LiveLinkHubProvider.cpp
--- a/Engine/Plugins/Animation/LiveLinkHub/Source/LiveLinkHub/Private/Clients/LiveLinkHubProvider.cpp
+++ b/Engine/Plugins/Animation/LiveLinkHub/Source/LiveLinkHub/Private/Clients/LiveLinkHubProvider.cpp
@@ -110,15 +110,40 @@ bool FLiveLinkHubProvider::ShouldTransmitToSubject_AnyThread(FName SubjectName,
 	return ShouldTransmitToClient_AnyThread(Address, AdditionalFilter);
 }
 
-void FLiveLinkHubProvider::SetTimecodeSettings(FLiveLinkHubTimecodeSettings InSettings)
+void FLiveLinkHubProvider::UpdateTimecodeSettings(const FLiveLinkHubTimecodeSettings& InSettings, const FLiveLinkHubClientId& ClientId)
 {
-	TimecodeSettings = MoveTemp(InSettings);
-	SendTimecodeSettings();
+	SendTimecodeSettings(InSettings, ClientId);
 }
 
-void FLiveLinkHubProvider::SendTimecodeSettings()
+void FLiveLinkHubProvider::ResetTimecodeSettings(const FLiveLinkHubClientId& ClientId)
 {
-	SendMessageToEnabledClients(FMessageEndpoint::MakeMessage<FLiveLinkHubTimecodeSettings>(TimecodeSettings));
+	// Sending settings with ELiveLinkHubTimecodeSource::NotDefined will reset the timecode on the client.
+	SendTimecodeSettings(FLiveLinkHubTimecodeSettings{}, ClientId);
+}
+
+void FLiveLinkHubProvider::SendTimecodeSettings(const FLiveLinkHubTimecodeSettings& InSettings, const FLiveLinkHubClientId& ClientId)
+{
+	if (ClientId.IsValid())
+	{
+		TArray<FMessageAddress> AllAddresses;
+		GetConnectedAddresses(AllAddresses);
+
+		FMessageAddress* TargetAddress = AllAddresses.FindByPredicate([this, ClientId](const FMessageAddress& Address)
+		{
+			return AddressToIdCache.FindChecked(Address) == ClientId;
+		});
+
+		if (TargetAddress)
+		{
+			SendMessage(FMessageEndpoint::MakeMessage<FLiveLinkHubTimecodeSettings>(InSettings));
+		}
+	}
+	else
+	{
+		// Invalid ID means we're broadcasting to all clients.
+		SendMessageToEnabledClients(FMessageEndpoint::MakeMessage<FLiveLinkHubTimecodeSettings>(InSettings));
+	}
+
 }
 
 void FLiveLinkHubProvider::AddRestoredClient(FLiveLinkHubUEClientInfo& RestoredClientInfo)
@@ -230,6 +255,8 @@ void FLiveLinkHubProvider::HandleHubConnectMessage(const FLiveLinkHubConnectMess
 		NewClient.IPAddress = LiveLinkHubProviderUtils::GetIPAddress(ConnectionAddress);
 
 		const FLiveLinkHubClientId NewClientId = NewClient.Id;
+		UpdatedClient = NewClientId; // Set this so we can update the client's timecode provider down below.
+
 		{
 			FWriteScopeLock Locker(ClientsMapLock);
 			AddressToIdCache.FindOrAdd(ConnectionAddress) = NewClientId;
@@ -254,6 +281,18 @@ void FLiveLinkHubProvider::HandleHubConnectMessage(const FLiveLinkHubConnectMess
 			OnClientEventDelegate.Broadcast(NewClientId, EClientEventType::Discovered);
 		}
 	}
+
+	// Update the timecode provider when a client establishes connection.
+	if (const TSharedPtr<ILiveLinkHubSessionManager> Manager = SessionManager.Pin())
+	{
+		if (const TSharedPtr<ILiveLinkHubSession> ActiveSession = Manager->GetCurrentSession())
+		{
+			if (ActiveSession->ShouldUseLiveLinkHubAsTimecodeSource())
+			{
+				SendTimecodeSettings(ActiveSession->GetTimecodeSettings(), *UpdatedClient);
+			}
+		}
+	}
 }
 	
 void FLiveLinkHubProvider::HandleClientInfoMessage(const FLiveLinkClientInfoMessage& Message, const TSharedRef<IMessageContext, ESPMode::ThreadSafe>& Context)
@@ -269,7 +308,16 @@ void FLiveLinkHubProvider::HandleClientInfoMessage(const FLiveLinkClientInfoMess
 		}
 	}
 
-	SendTimecodeSettings();
+	if (const TSharedPtr<ILiveLinkHubSessionManager> Manager = SessionManager.Pin())
+	{
+		if (const TSharedPtr<ILiveLinkHubSession> ActiveSession = Manager->GetCurrentSession())
+		{
+			if (ActiveSession->ShouldUseLiveLinkHubAsTimecodeSource() && ClientId.IsValid())
+			{
+				SendTimecodeSettings(ActiveSession->GetTimecodeSettings(), ClientId);
+			}
+		}
+	}
 
 	if (ClientId.IsValid())
 	{
@@ -447,10 +495,24 @@ void FLiveLinkHubProvider::SetClientEnabled(FLiveLinkHubClientId Client, bool bI
 		}
 	}
 
-	// No need to update timecode settings if client was disabled.
-	if (bInEnable)
+	if (const TSharedPtr<ILiveLinkHubSessionManager> Manager = SessionManager.Pin())
 	{
-		SendTimecodeSettings();
+		if (const TSharedPtr<ILiveLinkHubSession> ActiveSession = Manager->GetCurrentSession())
+		{
+			if (ActiveSession->ShouldUseLiveLinkHubAsTimecodeSource())
+			{
+				if (bInEnable)
+				{
+					// Enabling client, send it up to date timecode settings.
+					SendTimecodeSettings(ActiveSession->GetTimecodeSettings(), Client);
+				}
+				else
+				{
+					// Disabling it, so reset its timecode provider.
+					ResetTimecodeSettings(Client);
+				}
+			}
+		}
 	}
 }
 
diff --git a/Engine/Plugins/Animation/LiveLinkHub/Source/LiveLinkHub/Private/Clients/LiveLinkHubProvider.h b/Engine/Plugins/Animation/LiveLinkHub/Source/LiveLinkHub/Private/Clients/LiveLinkHubProvider.h
--- a/Engine/Plugins/Animation/LiveLinkHub/Source/LiveLinkHub/Private/Clients/LiveLinkHubProvider.h
+++ b/Engine/Plugins/Animation/LiveLinkHub/Source/LiveLinkHub/Private/Clients/LiveLinkHubProvider.h
@@ -54,12 +54,17 @@ public:
 	/** Retrieve the existing client map. */
 	const TMap<FLiveLinkHubClientId, FLiveLinkHubUEClientInfo>& GetClientsMap() const { return ClientsMap; }
 
-	/** Timecode settings that should be shared to connected editors. */
-	void SetTimecodeSettings(FLiveLinkHubTimecodeSettings InSettings);
-
-	/** Retrieve current timecode settings. */
-	const FLiveLinkHubTimecodeSettings& GetTimecodeSettings() const { return TimecodeSettings; }
+	/** 
+	 * Timecode settings that should be shared to connected editors. 
+	 * @Note If ClientId is not provided or invalid, the message will be broadcast to all connected clients.
+	 */
+	void UpdateTimecodeSettings(const FLiveLinkHubTimecodeSettings& InSettings, const FLiveLinkHubClientId& ClientId = FLiveLinkHubClientId{});
 
+	/**
+	 * Reset timecode settings for all connected clients. 
+	 * @Note If ClientId is not provided or invalid, the message will be broadcast to all connected clients.
+	 */
+	void ResetTimecodeSettings(const FLiveLinkHubClientId& ClientId = FLiveLinkHubClientId{});
 private:
 	/** Handle a connection message resulting from a livelink hub message bus source connecting to this provider. */
 	void HandleHubConnectMessage(const FLiveLinkHubConnectMessage& Message, const TSharedRef<IMessageContext, ESPMode::ThreadSafe>& Context);
@@ -68,7 +73,7 @@ private:
 	void HandleClientInfoMessage(const FLiveLinkClientInfoMessage& Message, const TSharedRef<IMessageContext, ESPMode::ThreadSafe>& Context);
 
 	/** Send timecode settings to connected Live Link Hub provider. */
-	void SendTimecodeSettings();
+	void SendTimecodeSettings(const FLiveLinkHubTimecodeSettings& InSettings, const FLiveLinkHubClientId& ClientId);
 
 	/** Send a message to clients that are connected and enabled through the hub clients list. */
 	template<typename MessageType>
@@ -123,8 +128,6 @@ private:
 	TWeakPtr<ILiveLinkHubSessionManager> SessionManager;
 	/** Cache used to retrieve the client id from a message bus address. */
 	TMap<FMessageAddress, FLiveLinkHubClientId> AddressToIdCache;
-	/** Cached value of timecode connection settings. */
-	FLiveLinkHubTimecodeSettings TimecodeSettings;
 
 	/** Lock used to access the clients map from different threads. */
 	mutable FRWLock ClientsMapLock;
diff --git a/Engine/Plugins/Animation/LiveLinkHub/Source/LiveLinkHub/Private/Session/LiveLinkHubSession.h b/Engine/Plugins/Animation/LiveLinkHub/Source/LiveLinkHub/Private/Session/LiveLinkHubSession.h
--- a/Engine/Plugins/Animation/LiveLinkHub/Source/LiveLinkHub/Private/Session/LiveLinkHubSession.h
+++ b/Engine/Plugins/Animation/LiveLinkHub/Source/LiveLinkHub/Private/Session/LiveLinkHubSession.h
@@ -43,6 +43,18 @@ public:
 
 	/** Get the list of clients in this session (The list of clients that can receive data from the hub) */
 	virtual TArray<FLiveLinkHubClientId> GetSessionClients() const = 0;
+
+	/** Returns whether livelinkhub should be used as a timecode source for connected clients. */
+	virtual bool ShouldUseLiveLinkHubAsTimecodeSource() const = 0;
+
+	/** Set whether livelinkhub should be used as a timecode source for connected clients.  */
+	virtual void SetUseLiveLinkHubAsTimecodeSource(bool bUseLiveLinkHubAsTimecodeSource) = 0;
+
+	/** Get the timecode settings for the current session. */
+	virtual FLiveLinkHubTimecodeSettings GetTimecodeSettings() const = 0;
+
+	/** Set the timecode settings for the current session. */
+	virtual void SetTimecodeSettings(const FLiveLinkHubTimecodeSettings& TimecodeSettings) const = 0;
 };
 
 class FLiveLinkHubSession : public ILiveLinkHubSession, public TSharedFromThis<FLiveLinkHubSession>
@@ -80,6 +92,11 @@ public:
 					FWriteScopeLock Locker(SessionDataLock);
 					CachedSessionClients.Add(Client);
 				}
+
+				if (ShouldUseLiveLinkHubAsTimecodeSource())
+				{
+					LiveLinkProvider->UpdateTimecodeSettings(GetTimecodeSettings(), Client);
+				}
 			}
 		}
 
@@ -90,6 +107,11 @@ public:
 	{
 		check(IsInGameThread());
 
+		if (TSharedPtr<FLiveLinkHubProvider> LiveLinkProvider = FModuleManager::Get().GetModuleChecked<FLiveLinkHubModule>("LiveLinkHub").GetLiveLinkProvider())
+		{
+			LiveLinkProvider->ResetTimecodeSettings(Client);
+		}
+
 		{
 			FWriteScopeLock Locker(SessionDataLock);
 			CachedSessionClients.Remove(Client);
@@ -104,6 +126,30 @@ public:
 		return CachedSessionClients.Contains(Client);
 	}
 
+	virtual bool ShouldUseLiveLinkHubAsTimecodeSource() const override
+	{
+		FReadScopeLock Locker(SessionDataLock);
+		return SessionData->bUseLiveLinkHubAsTimecodeSource;
+	}
+
+	virtual void SetUseLiveLinkHubAsTimecodeSource(bool bUseLiveLinkHubAsTimecodeSource) override
+	{
+		FWriteScopeLock Locker(SessionDataLock);
+		SessionData->bUseLiveLinkHubAsTimecodeSource = bUseLiveLinkHubAsTimecodeSource;
+	}
+
+	virtual FLiveLinkHubTimecodeSettings GetTimecodeSettings() const override
+	{
+		FReadScopeLock Locker(SessionDataLock);
+		return SessionData->TimecodeSettings;
+	}
+
+	virtual void SetTimecodeSettings(const FLiveLinkHubTimecodeSettings& TimecodeSettings) const override
+	{
+		FWriteScopeLock Locker(SessionDataLock);
+		SessionData->TimecodeSettings = TimecodeSettings;
+	}
+
 	void AddRestoredClient(FLiveLinkHubUEClientInfo& InOutRestoredClientInfo)
 	{
 		if (const TSharedPtr<FLiveLinkHubProvider> LiveLinkProvider = FModuleManager::Get().GetModuleChecked<FLiveLinkHubModule>("LiveLinkHub").GetLiveLinkProvider())
diff --git a/Engine/Plugins/Animation/LiveLinkHub/Source/LiveLinkHub/Private/Session/LiveLinkHubSessionData.h b/Engine/Plugins/Animation/LiveLinkHub/Source/LiveLinkHub/Private/Session/LiveLinkHubSessionData.h
--- a/Engine/Plugins/Animation/LiveLinkHub/Source/LiveLinkHub/Private/Session/LiveLinkHubSessionData.h
+++ b/Engine/Plugins/Animation/LiveLinkHub/Source/LiveLinkHub/Private/Session/LiveLinkHubSessionData.h
@@ -33,4 +33,8 @@ public:
 	/** Timecode settings for the live link hub. */
 	UPROPERTY()
 	FLiveLinkHubTimecodeSettings TimecodeSettings;
+
+	/** Whether the hub should be used as a timecode source for connected clients. */
+	UPROPERTY()
+	bool bUseLiveLinkHubAsTimecodeSource = false;
 };
diff --git a/Engine/Plugins/Animation/LiveLinkHub/Source/LiveLinkHub/Private/Session/LiveLinkHubSessionManager.h b/Engine/Plugins/Animation/LiveLinkHub/Source/LiveLinkHub/Private/Session/LiveLinkHubSessionManager.h
--- a/Engine/Plugins/Animation/LiveLinkHub/Source/LiveLinkHub/Private/Session/LiveLinkHubSessionManager.h
+++ b/Engine/Plugins/Animation/LiveLinkHub/Source/LiveLinkHub/Private/Session/LiveLinkHubSessionManager.h
@@ -169,9 +169,6 @@ public:
 			LiveLinkHubSessionData->Clients.Add(ClientKeyVal.Value);
 		}
 
-		const FLiveLinkHubTimecodeSettings& TimecodeSettings = LiveLinkProvider->GetTimecodeSettings();
-		LiveLinkHubSessionData->TimecodeSettings = TimecodeSettings;
-		
 		if (!SavePath.IsEmpty())
 		{
 			LastConfigPath = SavePath;
@@ -268,7 +265,7 @@ private:
 			}
 
 			SessionData->TimecodeSettings.AssignTimecodeSettingsAsProviderToEngine();
-			LiveLinkProvider->SetTimecodeSettings(SessionData->TimecodeSettings);
+			LiveLinkProvider->UpdateTimecodeSettings(SessionData->TimecodeSettings);
 
 			TSharedPtr<FLiveLinkHubSession> CurrentSessionPtr;
 			{
diff --git a/Engine/Plugins/Animation/LiveLinkHub/Source/LiveLinkHub/Private/UI/Widgets/SLiveLinkTimecode.cpp b/Engine/Plugins/Animation/LiveLinkHub/Source/LiveLinkHub/Private/UI/Widgets/SLiveLinkTimecode.cpp
--- a/Engine/Plugins/Animation/LiveLinkHub/Source/LiveLinkHub/Private/UI/Widgets/SLiveLinkTimecode.cpp
+++ b/Engine/Plugins/Animation/LiveLinkHub/Source/LiveLinkHub/Private/UI/Widgets/SLiveLinkTimecode.cpp
@@ -99,12 +99,12 @@ namespace UE::LiveLinkTimecode::Private
 
 FSlateColor SLiveLinkTimecode::GetTimecodeStatusColor() const
 {
-	return bIsTimecodeSource ? FSlateColor(FColor::Green) : FSlateColor(FColor::Yellow);
+	return bCachedIsTimecodeSource ? FSlateColor(FColor::Green) : FSlateColor(FColor::Yellow);
 }
 
 FText SLiveLinkTimecode::GetTimecodeTooltip() const
 {
-	if (bIsTimecodeSource)
+	if (bCachedIsTimecodeSource)
 	{
 		return LOCTEXT("LiveLinkTimeCode_IsConnected", "Sending timecode data to connected editors.");
 	}
@@ -113,34 +113,50 @@ FText SLiveLinkTimecode::GetTimecodeTooltip() const
 
 void SLiveLinkTimecode::OnEnableTimecodeToggled()
 {
-	bIsTimecodeSource = !bIsTimecodeSource;
-	SendUpdatedTimecodeToEditor();
+	bCachedIsTimecodeSource = !bCachedIsTimecodeSource;
+
+	if (TSharedPtr<ILiveLinkHubSessionManager> SessionManager = FModuleManager::Get().GetModuleChecked<FLiveLinkHubModule>("LiveLinkHub").GetSessionManager())
+	{
+		SessionManager->GetCurrentSession()->SetUseLiveLinkHubAsTimecodeSource(bCachedIsTimecodeSource);
+	}
+
+	if (bCachedIsTimecodeSource)
+	{
+		SendUpdatedTimecodeToEditor(MakeTimecodeSettings());
+	}
+	else
+	{
+		// If we're disabling LLH as 
+		const TSharedPtr<FLiveLinkHub> LiveLinkHub = FModuleManager::Get().GetModuleChecked<FLiveLinkHubModule>("LiveLinkHub").GetLiveLinkHub();
+		if (LiveLinkHub.IsValid())
+		{
+			if (TSharedPtr<FLiveLinkHubProvider> Provider = LiveLinkHub->GetLiveLinkProvider())
+			{
+				Provider->ResetTimecodeSettings();
+			}
+		}
+	}
 }
 
-void SLiveLinkTimecode::SendUpdatedTimecodeToEditor()
+FLiveLinkHubTimecodeSettings SLiveLinkTimecode::MakeTimecodeSettings() const
 {
 	FLiveLinkHubTimecodeSettings Settings;
 	using namespace UE::LiveLinkTimecode::Private;
 
-	if (!bIsTimecodeSource)
-	{
-		Settings.Source = ELiveLinkHubTimecodeSource::NotDefined;
-	}
-
 	auto GetSystemTimeSource = [this]() -> TOptional<FLiveLinkHubTimecodeSettings>
-	{
-		for (const TPair<FName,FTimecodeMenuItem>& Source : StaticTimecodeMenu)
 		{
-			if (ActiveTimecodeSource == Source.Key)
+			for (const TPair<FName, FTimecodeMenuItem>& Source : StaticTimecodeMenu)
 			{
-				FLiveLinkHubTimecodeSettings Settings;
-				Settings.Source = ELiveLinkHubTimecodeSource::SystemTimeEditor;
-				Settings.DesiredFrameRate = Source.Value.Rate;
-				return Settings;
+				if (ActiveTimecodeSource == Source.Key)
+				{
+					FLiveLinkHubTimecodeSettings Settings;
+					Settings.Source = ELiveLinkHubTimecodeSource::SystemTimeEditor;
+					Settings.DesiredFrameRate = Source.Value.Rate;
+					return Settings;
+				}
 			}
-		}
-		return {};
-	};
+			return {};
+		};
 
 	if (TOptional<FLiveLinkHubTimecodeSettings> SystemTimeSource = GetSystemTimeSource())
 	{
@@ -150,38 +166,32 @@ void SLiveLinkTimecode::SendUpdatedTimecodeToEditor()
 	{
 		Settings.Source = ELiveLinkHubTimecodeSource::UseSubjectName;
 	}
+
 	Settings.SubjectName = ActiveTimecodeSource;
 
-	// Apply the timecode settings to the local instance too.
-	Settings.AssignTimecodeSettingsAsProviderToEngine();
+	return Settings;
+}
+
 
+void SLiveLinkTimecode::SendUpdatedTimecodeToEditor(const FLiveLinkHubTimecodeSettings& TimecodeSettings)
+{
 	const TSharedPtr<FLiveLinkHub> LiveLinkHub = FModuleManager::Get().GetModuleChecked<FLiveLinkHubModule>("LiveLinkHub").GetLiveLinkHub();
 	if (LiveLinkHub.IsValid())
 	{
-		TSharedPtr<FLiveLinkHubProvider> Provider = LiveLinkHub->GetLiveLinkProvider();
-		if (Provider.IsValid())
+		if (TSharedPtr<FLiveLinkHubProvider> Provider = LiveLinkHub->GetLiveLinkProvider())
 		{
-			Provider->SetTimecodeSettings(MoveTemp(Settings));
+			Provider->UpdateTimecodeSettings(TimecodeSettings);
 		}
 	}
 }
 
-void SLiveLinkTimecode::UpdateTimecodeFromSettings()
+void SLiveLinkTimecode::UpdateTimecodeSettingsFromSession(const TSharedRef<ILiveLinkHubSession>& InSession)
 {
-	FLiveLinkHubTimecodeSettings Settings;
 	using namespace UE::LiveLinkTimecode::Private;
 
-	const TSharedPtr<FLiveLinkHub> LiveLinkHub = FModuleManager::Get().GetModuleChecked<FLiveLinkHubModule>("LiveLinkHub").GetLiveLinkHub();
-	if (LiveLinkHub.IsValid())
-	{
-		TSharedPtr<FLiveLinkHubProvider> Provider = LiveLinkHub->GetLiveLinkProvider();
-		if (Provider.IsValid())
-		{
-			Settings = Provider->GetTimecodeSettings();
-		}
-	}
+	const FLiveLinkHubTimecodeSettings Settings = InSession->GetTimecodeSettings();
+	bCachedIsTimecodeSource = InSession->ShouldUseLiveLinkHubAsTimecodeSource();
 
-	bIsTimecodeSource = Settings.Source != ELiveLinkHubTimecodeSource::NotDefined;
 	ActiveTimecodeSource = Settings.SubjectName.IsNone() ? System24fps : Settings.SubjectName;
 }
 
@@ -193,7 +203,20 @@ void SLiveLinkTimecode::SetTimecodeSource(const FName SourceId)
 	}
 
 	ActiveTimecodeSource = SourceId;
-	SendUpdatedTimecodeToEditor();
+	
+	const FLiveLinkHubTimecodeSettings Settings = MakeTimecodeSettings();
+	Settings.AssignTimecodeSettingsAsProviderToEngine();
+
+	const TSharedPtr<FLiveLinkHub> LiveLinkHub = FModuleManager::Get().GetModuleChecked<FLiveLinkHubModule>("LiveLinkHub").GetLiveLinkHub();
+	if (LiveLinkHub.IsValid())
+	{
+		LiveLinkHub->GetSessionManager()->GetCurrentSession()->SetTimecodeSettings(Settings);
+	}
+
+	if (bCachedIsTimecodeSource)
+	{
+		SendUpdatedTimecodeToEditor(Settings);
+	}
 }
 
 TSharedRef<SWidget> SLiveLinkTimecode::MakeMenu()
@@ -207,7 +230,7 @@ TSharedRef<SWidget> SLiveLinkTimecode::MakeMenu()
 		FUIAction(
 			FExecuteAction::CreateSP(this, &SLiveLinkTimecode::OnEnableTimecodeToggled),
 			FCanExecuteAction::CreateLambda([] { return true; }),
-			FIsActionChecked::CreateLambda([this] { return bIsTimecodeSource; })),
+			FIsActionChecked::CreateLambda([this] { return bCachedIsTimecodeSource; })),
 		NAME_None,
 		EUserInterfaceActionType::ToggleButton);
 
@@ -260,15 +283,15 @@ void SLiveLinkTimecode::Construct(const FArguments& InArgs)
 	TSharedPtr<ILiveLinkHubSessionManager> SessionManager = LiveLinkHubModule.GetSessionManager();
 	if (SessionManager.IsValid())
 	{
-		SessionManager->OnActiveSessionChanged().AddLambda([this](const TSharedRef<ILiveLinkHubSession>&)
+		UpdateTimecodeSettingsFromSession(SessionManager->GetCurrentSession().ToSharedRef());
+
+		SessionManager->OnActiveSessionChanged().AddLambda([this](const TSharedRef<ILiveLinkHubSession>& InSession)
 		{
 			// Update the UI when a config is loaded.
-			UpdateTimecodeFromSettings();
+				UpdateTimecodeSettingsFromSession(InSession);
 		});
 	}
 
-	UpdateTimecodeFromSettings();
-	
 	check(WorkingClient);
 	ChildSlot
 	[
diff --git a/Engine/Plugins/Animation/LiveLinkHub/Source/LiveLinkHub/Private/UI/Widgets/SLiveLinkTimecode.h b/Engine/Plugins/Animation/LiveLinkHub/Source/LiveLinkHub/Private/UI/Widgets/SLiveLinkTimecode.h
--- a/Engine/Plugins/Animation/LiveLinkHub/Source/LiveLinkHub/Private/UI/Widgets/SLiveLinkTimecode.h
+++ b/Engine/Plugins/Animation/LiveLinkHub/Source/LiveLinkHub/Private/UI/Widgets/SLiveLinkTimecode.h
@@ -6,6 +6,7 @@
 
 class FName;
 class FLiveLinkClient;
+struct FLiveLinkHubTimecodeSettings;
 
 /**
  * Utility class to display timecode status on the menu bar. Provides user with
@@ -42,13 +43,15 @@ private:
 	void SetTimecodeSource(const FName SourceId);
 
 	/** Sends the current timecode settings to the connected editor(s). */
-	void SendUpdatedTimecodeToEditor();
+	void SendUpdatedTimecodeToEditor(const FLiveLinkHubTimecodeSettings& TimecodeSettings);
+
+	FLiveLinkHubTimecodeSettings MakeTimecodeSettings() const;
 
 	/** Update the local timecode properties from the timecode settings. */
-	void UpdateTimecodeFromSettings();
+	void UpdateTimecodeSettingsFromSession(const TSharedRef<class ILiveLinkHubSession>& InSession);
 	
 	/** Do we have timecode source enabled */
-	bool bIsTimecodeSource = false;
+	bool bCachedIsTimecodeSource = false;
 
 	/** Name of the active timecode source. */
 	FName ActiveTimecodeSource;
