Add UI to control timestep via Live Link Hub.

Custom Time Step will only be applied to connected editors if both a framerate and a subject are selected.

diff --git a/Engine/Plugins/Animation/LiveLink/Source/LiveLink/Private/LiveLinkCustomTimeStep.cpp b/Engine/Plugins/Animation/LiveLink/Source/LiveLink/Private/LiveLinkCustomTimeStep.cpp
--- a/Engine/Plugins/Animation/LiveLink/Source/LiveLink/Private/LiveLinkCustomTimeStep.cpp
+++ b/Engine/Plugins/Animation/LiveLink/Source/LiveLink/Private/LiveLinkCustomTimeStep.cpp
@@ -3,6 +3,7 @@
 #include "LiveLinkCustomTimeStep.h"
 
 #include "Features/IModularFeatures.h"
+#include "GenlockedCustomTimeStep.h"
 #include "ILiveLinkClient.h"
 #include "Misc/App.h"
 
@@ -112,6 +113,11 @@ bool ULiveLinkCustomTimeStep::IsLastSyncDataValid() const
 	return true;
 }
 
+uint32 ULiveLinkCustomTimeStep::GetExpectedSyncCountDelta() const
+{
+	return UGenlockedCustomTimeStep::GetExpectedSyncCountDelta();
+}
+
 bool ULiveLinkCustomTimeStep::WaitForSync()
 {	
 	EventWaitForSync->Trigger();
diff --git a/Engine/Plugins/Animation/LiveLink/Source/LiveLink/Public/LiveLinkCustomTimeStep.h b/Engine/Plugins/Animation/LiveLink/Source/LiveLink/Public/LiveLinkCustomTimeStep.h
--- a/Engine/Plugins/Animation/LiveLink/Source/LiveLink/Public/LiveLinkCustomTimeStep.h
+++ b/Engine/Plugins/Animation/LiveLink/Source/LiveLink/Public/LiveLinkCustomTimeStep.h
@@ -50,11 +50,12 @@ public:
 	LIVELINK_API virtual uint32 GetLastSyncCountDelta() const override;
 	LIVELINK_API virtual bool IsLastSyncDataValid() const override;
 	LIVELINK_API virtual bool WaitForSync() override;
+	LIVELINK_API virtual uint32 GetExpectedSyncCountDelta() const override;
 
 	//~ Begin UObject Interface
-	virtual void BeginDestroy() override;
+	LIVELINK_API virtual void BeginDestroy() override;
 #if WITH_EDITOR
-	virtual void PostEditChangeProperty(FPropertyChangedEvent& PropertyChangedEvent) override;
+	LIVELINK_API virtual void PostEditChangeProperty(FPropertyChangedEvent& PropertyChangedEvent) override;
 #endif
 	//~End UObject Interface
 
diff --git a/Engine/Plugins/Animation/LiveLinkHub/Source/LiveLinkHub/Private/Clients/LiveLinkHubProvider.cpp b/Engine/Plugins/Animation/LiveLinkHub/Source/LiveLinkHub/Private/Clients/LiveLinkHubProvider.cpp
--- a/Engine/Plugins/Animation/LiveLinkHub/Source/LiveLinkHub/Private/Clients/LiveLinkHubProvider.cpp
+++ b/Engine/Plugins/Animation/LiveLinkHub/Source/LiveLinkHub/Private/Clients/LiveLinkHubProvider.cpp
@@ -121,6 +121,20 @@ void FLiveLinkHubProvider::ResetTimecodeSettings(const FLiveLinkHubClientId& Cli
 	SendTimecodeSettings(FLiveLinkHubTimecodeSettings{}, ClientId);
 }
 
+void FLiveLinkHubProvider::UpdateCustomTimeStepSettings(const FLiveLinkHubCustomTimeStepSettings& InSettings, const FLiveLinkHubClientId& ClientId)
+{
+	SendCustomTimeStepSettings(InSettings, ClientId);
+}
+
+void FLiveLinkHubProvider::ResetCustomTimeStepSettings(const FLiveLinkHubClientId& ClientId)
+{
+	// Setting the bResetCustomTimeStep flag will reset the CustomTimeStep on the client.
+	FLiveLinkHubCustomTimeStepSettings ResetCustomTimeStepSettings;
+	ResetCustomTimeStepSettings.bResetCustomTimeStep = true;
+
+	SendCustomTimeStepSettings(ResetCustomTimeStepSettings, ClientId);
+}
+
 void FLiveLinkHubProvider::SendTimecodeSettings(const FLiveLinkHubTimecodeSettings& InSettings, const FLiveLinkHubClientId& ClientId)
 {
 	if (ClientId.IsValid())
@@ -143,7 +157,30 @@ void FLiveLinkHubProvider::SendTimecodeSettings(const FLiveLinkHubTimecodeSettin
 		// Invalid ID means we're broadcasting to all clients.
 		SendMessageToEnabledClients(FMessageEndpoint::MakeMessage<FLiveLinkHubTimecodeSettings>(InSettings));
 	}
+}
+
+void FLiveLinkHubProvider::SendCustomTimeStepSettings(const FLiveLinkHubCustomTimeStepSettings& InSettings, const FLiveLinkHubClientId& ClientId)
+{
+	if (ClientId.IsValid())
+	{
+		TArray<FMessageAddress> AllAddresses;
+		GetConnectedAddresses(AllAddresses);
+
+		FMessageAddress* TargetAddress = AllAddresses.FindByPredicate([this, ClientId](const FMessageAddress& Address)
+			{
+				return AddressToIdCache.FindChecked(Address) == ClientId;
+			});
 
+		if (TargetAddress)
+		{
+			SendMessage(FMessageEndpoint::MakeMessage<FLiveLinkHubCustomTimeStepSettings>(InSettings));
+		}
+	}
+	else
+	{
+		// Invalid ID means we're broadcasting to all clients.
+		SendMessageToEnabledClients(FMessageEndpoint::MakeMessage<FLiveLinkHubCustomTimeStepSettings>(InSettings));
+	}
 }
 
 void FLiveLinkHubProvider::AddRestoredClient(FLiveLinkHubUEClientInfo& RestoredClientInfo)
@@ -290,6 +327,7 @@ void FLiveLinkHubProvider::HandleHubConnectMessage(const FLiveLinkHubConnectMess
 			if (ActiveSession->ShouldUseLiveLinkHubAsTimecodeSource())
 			{
 				SendTimecodeSettings(ActiveSession->GetTimecodeSettings(), *UpdatedClient);
+				SendCustomTimeStepSettings(ActiveSession->GetCustomTimeStepSettings(), *UpdatedClient);
 			}
 		}
 	}
@@ -503,15 +541,27 @@ void FLiveLinkHubProvider::SetClientEnabled(FLiveLinkHubClientId Client, bool bI
 			{
 				if (bInEnable)
 				{
-					// Enabling client, send it up to date timecode settings.
+					// Enabling client, send it up to date timecode and custom time step settings.
 					SendTimecodeSettings(ActiveSession->GetTimecodeSettings(), Client);
 				}
 				else
 				{
-					// Disabling it, so reset its timecode provider.
+					// Disabling it, so reset its timecode and custom time step.
 					ResetTimecodeSettings(Client);
 				}
 			}
+
+			if (ActiveSession->ShouldUseLiveLinkHubAsCustomTimeStepSource())
+			{
+				if (bInEnable)
+				{
+					SendCustomTimeStepSettings(ActiveSession->GetCustomTimeStepSettings(), Client);
+				}
+				else
+				{
+					ResetCustomTimeStepSettings(Client);
+				}
+			}
 		}
 	}
 }
diff --git a/Engine/Plugins/Animation/LiveLinkHub/Source/LiveLinkHub/Private/Clients/LiveLinkHubProvider.h b/Engine/Plugins/Animation/LiveLinkHub/Source/LiveLinkHub/Private/Clients/LiveLinkHubProvider.h
--- a/Engine/Plugins/Animation/LiveLinkHub/Source/LiveLinkHub/Private/Clients/LiveLinkHubProvider.h
+++ b/Engine/Plugins/Animation/LiveLinkHub/Source/LiveLinkHub/Private/Clients/LiveLinkHubProvider.h
@@ -65,6 +65,16 @@ public:
 	 * @Note If ClientId is not provided or invalid, the message will be broadcast to all connected clients.
 	 */
 	void ResetTimecodeSettings(const FLiveLinkHubClientId& ClientId = FLiveLinkHubClientId{});
+
+	/** Frame Lock settings that should be shared to connected editors. */
+	void UpdateCustomTimeStepSettings(const FLiveLinkHubCustomTimeStepSettings& InSettings, const FLiveLinkHubClientId& ClientId = FLiveLinkHubClientId{});
+
+	/** 
+	 * Reset Frame Lock settings on connected editors. 
+	 * @Note If ClientId is not provided or invalid, the message will be broadcast to all connected clients.
+	 */
+	void ResetCustomTimeStepSettings(const FLiveLinkHubClientId& ClientId = FLiveLinkHubClientId{});
+
 private:
 	/** Handle a connection message resulting from a livelink hub message bus source connecting to this provider. */
 	void HandleHubConnectMessage(const FLiveLinkHubConnectMessage& Message, const TSharedRef<IMessageContext, ESPMode::ThreadSafe>& Context);
@@ -72,9 +82,12 @@ private:
 	/** Handle a client info message being received. Happens when new information about a client is received (ie. Client has changed map) */
 	void HandleClientInfoMessage(const FLiveLinkClientInfoMessage& Message, const TSharedRef<IMessageContext, ESPMode::ThreadSafe>& Context);
 
-	/** Send timecode settings to connected Live Link Hub provider. */
+	/** Send timecode settings to connected UE Clients. */
 	void SendTimecodeSettings(const FLiveLinkHubTimecodeSettings& InSettings, const FLiveLinkHubClientId& ClientId);
 
+	/** Send CustomTimeStep settings to connected UE Clients. */
+	void SendCustomTimeStepSettings(const FLiveLinkHubCustomTimeStepSettings& InSettings, const FLiveLinkHubClientId& ClientId);
+
 	/** Send a message to clients that are connected and enabled through the hub clients list. */
 	template<typename MessageType>
 	void SendMessageToEnabledClients(MessageType* Message)
diff --git a/Engine/Plugins/Animation/LiveLinkHub/Source/LiveLinkHub/Private/Session/LiveLinkHubSession.h b/Engine/Plugins/Animation/LiveLinkHub/Source/LiveLinkHub/Private/Session/LiveLinkHubSession.h
--- a/Engine/Plugins/Animation/LiveLinkHub/Source/LiveLinkHub/Private/Session/LiveLinkHubSession.h
+++ b/Engine/Plugins/Animation/LiveLinkHub/Source/LiveLinkHub/Private/Session/LiveLinkHubSession.h
@@ -55,6 +55,18 @@ public:
 
 	/** Set the timecode settings for the current session. */
 	virtual void SetTimecodeSettings(const FLiveLinkHubTimecodeSettings& TimecodeSettings) const = 0;
+
+	/** Returns whether livelinkhub should be used as a CustomTimeStep source for connected clients. */
+	virtual bool ShouldUseLiveLinkHubAsCustomTimeStepSource() const = 0;
+
+	/** Set whether livelinkhub should be used as a timecode source for connected clients.  */
+	virtual void SetShouldUseLiveLinkHubAsCustomTimeStepSource(bool bUseLiveLinkHubAsCustomTimeStepSource) = 0;
+
+	/** Get the CustomTimeStep settings for the current session. */
+	virtual FLiveLinkHubCustomTimeStepSettings GetCustomTimeStepSettings() const = 0;
+
+	/** Set the CustomTimeStep settings for the current session. */
+	virtual void SetCustomTimeStepSettings(const FLiveLinkHubCustomTimeStepSettings& CustomTimeStepSettings) const = 0;
 };
 
 class FLiveLinkHubSession : public ILiveLinkHubSession, public TSharedFromThis<FLiveLinkHubSession>
@@ -96,6 +108,7 @@ public:
 				if (ShouldUseLiveLinkHubAsTimecodeSource())
 				{
 					LiveLinkProvider->UpdateTimecodeSettings(GetTimecodeSettings(), Client);
+					LiveLinkProvider->UpdateCustomTimeStepSettings(GetCustomTimeStepSettings(), Client);
 				}
 			}
 		}
@@ -110,6 +123,7 @@ public:
 		if (TSharedPtr<FLiveLinkHubProvider> LiveLinkProvider = FModuleManager::Get().GetModuleChecked<FLiveLinkHubModule>("LiveLinkHub").GetLiveLinkProvider())
 		{
 			LiveLinkProvider->ResetTimecodeSettings(Client);
+			LiveLinkProvider->ResetCustomTimeStepSettings(Client);
 		}
 
 		{
@@ -150,6 +164,30 @@ public:
 		SessionData->TimecodeSettings = TimecodeSettings;
 	}
 
+	virtual bool ShouldUseLiveLinkHubAsCustomTimeStepSource() const
+	{
+		FReadScopeLock Locker(SessionDataLock);
+		return SessionData->bUseLiveLinkHubAsCustomTimeStepSource;
+	}
+
+	virtual void SetShouldUseLiveLinkHubAsCustomTimeStepSource(bool bUseLiveLinkHubAsCustomTimeStepSource)
+	{
+		FWriteScopeLock Locker(SessionDataLock);
+		SessionData->bUseLiveLinkHubAsCustomTimeStepSource = bUseLiveLinkHubAsCustomTimeStepSource;
+	}
+
+	virtual FLiveLinkHubCustomTimeStepSettings GetCustomTimeStepSettings() const
+	{
+		FReadScopeLock Locker(SessionDataLock);
+		return SessionData->CustomTimeStepSettings;
+	}
+
+	virtual void SetCustomTimeStepSettings(const FLiveLinkHubCustomTimeStepSettings& CustomTimeStepSettings) const
+	{
+		FWriteScopeLock Locker(SessionDataLock);
+		SessionData->CustomTimeStepSettings = CustomTimeStepSettings;
+	}
+
 	void AddRestoredClient(FLiveLinkHubUEClientInfo& InOutRestoredClientInfo)
 	{
 		if (const TSharedPtr<FLiveLinkHubProvider> LiveLinkProvider = FModuleManager::Get().GetModuleChecked<FLiveLinkHubModule>("LiveLinkHub").GetLiveLinkProvider())
diff --git a/Engine/Plugins/Animation/LiveLinkHub/Source/LiveLinkHub/Private/Session/LiveLinkHubSessionData.h b/Engine/Plugins/Animation/LiveLinkHub/Source/LiveLinkHub/Private/Session/LiveLinkHubSessionData.h
--- a/Engine/Plugins/Animation/LiveLinkHub/Source/LiveLinkHub/Private/Session/LiveLinkHubSessionData.h
+++ b/Engine/Plugins/Animation/LiveLinkHub/Source/LiveLinkHub/Private/Session/LiveLinkHubSessionData.h
@@ -30,11 +30,19 @@ public:
 	UPROPERTY()
 	TArray<FLiveLinkHubUEClientInfo> Clients;
 
-	/** Timecode settings for the live link hub. */
+	/** Timecode settings for the Live Link Hub. */
 	UPROPERTY()
 	FLiveLinkHubTimecodeSettings TimecodeSettings;
 
+	/** Frame Lock settings for the Live Link Hub. */
+	UPROPERTY()
+	FLiveLinkHubCustomTimeStepSettings CustomTimeStepSettings;
+
 	/** Whether the hub should be used as a timecode source for connected clients. */
 	UPROPERTY()
 	bool bUseLiveLinkHubAsTimecodeSource = false;
+
+	/** Whether the hub should be used as a CustomTimeStep source for connected clients. */
+	UPROPERTY()
+	bool bUseLiveLinkHubAsCustomTimeStepSource = false;
 };
diff --git a/Engine/Plugins/Animation/LiveLinkHub/Source/LiveLinkHub/Private/Session/LiveLinkHubSessionManager.h b/Engine/Plugins/Animation/LiveLinkHub/Source/LiveLinkHub/Private/Session/LiveLinkHubSessionManager.h
--- a/Engine/Plugins/Animation/LiveLinkHub/Source/LiveLinkHub/Private/Session/LiveLinkHubSessionManager.h
+++ b/Engine/Plugins/Animation/LiveLinkHub/Source/LiveLinkHub/Private/Session/LiveLinkHubSessionManager.h
@@ -267,6 +267,9 @@ private:
 			SessionData->TimecodeSettings.AssignTimecodeSettingsAsProviderToEngine();
 			LiveLinkProvider->UpdateTimecodeSettings(SessionData->TimecodeSettings);
 
+			SessionData->CustomTimeStepSettings.AssignCustomTimeStepToEngine();
+			LiveLinkProvider->UpdateCustomTimeStepSettings(SessionData->CustomTimeStepSettings);
+
 			TSharedPtr<FLiveLinkHubSession> CurrentSessionPtr;
 			{
 				FScopeLock Lock(&CurrentSessionCS);
diff --git a/Engine/Plugins/Animation/LiveLinkHub/Source/LiveLinkHub/Private/UI/Widgets/SLiveLinkTimecode.cpp b/Engine/Plugins/Animation/LiveLinkHub/Source/LiveLinkHub/Private/UI/Widgets/SLiveLinkTimecode.cpp
--- a/Engine/Plugins/Animation/LiveLinkHub/Source/LiveLinkHub/Private/UI/Widgets/SLiveLinkTimecode.cpp
+++ b/Engine/Plugins/Animation/LiveLinkHub/Source/LiveLinkHub/Private/UI/Widgets/SLiveLinkTimecode.cpp
@@ -12,6 +12,8 @@
 #include "Misc/App.h"
 #include "UI/Widgets/SLiveLinkTimecode.h"
 #include "Widgets/Input/SComboButton.h"
+#include "Widgets/Input/SNumericEntryBox.h"
+#include "Widgets/Layout/SSpacer.h"
 #include "Widgets/STimecode.h"
 
 #include "Clients/LiveLinkHubProvider.h"
@@ -45,7 +47,8 @@ namespace UE::LiveLinkTimecode::Private
 		FText Label;
 		FText ToolTip;
 	};
-	static const TMap<FName, FTimecodeMenuItem> StaticTimecodeMenu = []() {
+	static const TMap<FName, FTimecodeMenuItem> StaticTimecodeMenu = []()
+	{
 		TMap<FName, FTimecodeMenuItem> Definition = {
 			{System23976fps, {
 				FFrameRate(24000,1001),
@@ -95,20 +98,116 @@ namespace UE::LiveLinkTimecode::Private
 		};
 		return Definition;
 	}();
+
+	/** We only support a preset list of timecode values + named subjects. */
+	static FName CustomTimeStep23976fps = TEXT("CustomTimeStepSystemTime23976fps");
+	static FName CustomTimeStep24fps = TEXT("CustomTimeStepSystemTime24fps");
+	static FName CustomTimeStep25fps = TEXT("CustomTimeStepSystemTime25fps");
+	static FName CustomTimeStep2997fps  = TEXT("CustomTimeStepSystemTime2997fps");
+	static FName CustomTimeStep30fps    = TEXT("CustomTimeStepSystemTime30fps");
+	static FName CustomTimeStep48fps = TEXT("CustomTimeStepSystemTime48fps");
+	static FName CustomTimeStep50fps = TEXT("CustomTimeStepSystemTime50fps");
+	static FName CustomTimeStep5994fps = TEXT("CustomTimeStepSystemTime5994fps");
+	static FName CustomTimeStep60fps = TEXT("CustomTimeStepSystemTime60fps");
+	static FName CustomTimeStep100fps = TEXT("CustomTimeStepSystemTime100fps");
+	static FName CustomTimeStep120fps = TEXT("CustomTimeStepSystemTime120fps");
+	static FName CustomTimeStepUseTimecodeProvider = TEXT("CustomTimeStepUseTimecodeProvider");
+
+	static const TMap<FName, FTimecodeMenuItem> StaticCustomTimeStepMenu = []()
+	{
+		TMap<FName, FTimecodeMenuItem> Definition = {
+			{ CustomTimeStep23976fps,
+             { FFrameRate(24000, 1001),
+				LOCTEXT("LiveLinkHubCustomTimeStep23976fps", "23.976 ND fps"),
+				LOCTEXT("LiveLinkHubCustomTimeStep23976fps_Tooltip", "Use a 23.976 Non-drop FPS timecode for CustomTimeStep.") } },
+			{ CustomTimeStep24fps,
+             { FFrameRate(24, 1),
+				LOCTEXT("LiveLinkHubCustomTimeStep24fps", "24 fps"),
+				LOCTEXT("LiveLinkHubCustomTimeStep24fps_Tooltip", "Use a 24 FPS timecode for CustomTimeStep.") }                 },
+			{ CustomTimeStep25fps,
+             { FFrameRate(25, 1),
+				LOCTEXT("LiveLinkHubCustomTimeStep25fps", "25 fps"),
+				LOCTEXT("LiveLinkHubCustomTimeStep25fps_Tooltip", "Use a 25 FPS timecode for CustomTimeStep.") }                 },
+			{ CustomTimeStep2997fps,
+             { FFrameRate(30000, 1001),
+				LOCTEXT("LiveLinkHubCustomTimeStep2997fps", "29.97 ND fps"),
+				LOCTEXT("LiveLinkHubCustomTimeStep29976fps_Tooltip", "Use a 29.97 Non-drop FPS timecode for CustomTimeStep.") }  },
+			{ CustomTimeStep30fps,
+             { FFrameRate(30, 1),
+				LOCTEXT("LiveLinkHubCustomTimeStep30fps", "30 fps"),
+				LOCTEXT("LiveLinkHubCustomTimeStep30fps_Tooltip", "Use a 30 FPS timecode for CustomTimeStep.") }                 },
+			{ CustomTimeStep48fps,
+             { FFrameRate(48, 1),
+				LOCTEXT("LiveLinkHubCustomTimeStep48fps", "48 fps"),
+				LOCTEXT("LiveLinkHubCustomTimeStep48fps_Tooltip", "Use a 48 FPS timecode for CustomTimeStep.") }                 },
+			{ CustomTimeStep50fps,
+             { FFrameRate(50, 1),
+				LOCTEXT("LiveLinkHubCustomTimeStep50fps", "50 fps"),
+				LOCTEXT("LiveLinkHubCustomTimeStep50fps_Tooltip", "Use a 50 FPS timecode for CustomTimeStep.") }                 },
+			{ CustomTimeStep5994fps,
+             { FFrameRate(60000, 1001),
+				LOCTEXT("LiveLinkHubCustomTimeStep5994fps", "59.94 ND fps"),
+				LOCTEXT("LiveLinkHubCustomTimeStep5994fps_Tooltip", "Use a 59.94 Non-drp FPS timecode for CustomTimeStep.") }    },
+			{ CustomTimeStep60fps,
+             { FFrameRate(60, 1),
+				LOCTEXT("LiveLinkHubCustomTimeStep60fps", "60 fps"),
+				LOCTEXT("LiveLinkHubCustomTimeStep60fps_Tooltip", "Use a 60 FPS timecode for CustomTimeStep.") }                 },
+			{ CustomTimeStep100fps,
+             { FFrameRate(100, 1),
+				LOCTEXT("LiveLinkHubCustomTimeStep100fps", "100 fps"),
+				LOCTEXT("LiveLinkHubCustomTimeStep100fps_Tooltip", "Use a 100 FPS timecode for CustomTimeStep.") }               },
+			{
+             CustomTimeStep120fps,    { FFrameRate(120, 1),
+				  LOCTEXT("LiveLinkHubCustomTimeStep100fps", "120 fps"),
+				  LOCTEXT("LiveLinkHubCustomTimeStep100fps_Tooltip", "Use a 120 FPS timecode for CustomTimeStep.") }
+		}};
+		return Definition;
+	}();
+
 }
 
-FSlateColor SLiveLinkTimecode::GetTimecodeStatusColor() const
+FSlateColor SLiveLinkTimecode::GetStatusColor() const
 {
-	return bCachedIsTimecodeSource ? FSlateColor(FColor::Green) : FSlateColor(FColor::Yellow);
+	if (bCachedIsCustomTimeStepEnabled)
+	{
+		if (ActiveSubjectCustomTimeStep.IsNone() || ActiveCustomTimeStep.IsNone())
+		{
+			return FSlateColor(FColor::Red);
+		}
+	}
+
+	if (!bCachedIsTimecodeSource)
+	{
+		return FSlateColor(FColor::Yellow);
+	}
+
+	return FSlateColor(FColor::Green);
 }
 
 FText SLiveLinkTimecode::GetTimecodeTooltip() const
 {
-	if (bCachedIsTimecodeSource)
+	if (bCachedIsCustomTimeStepEnabled)
+	{
+		if (ActiveSubjectCustomTimeStep.IsNone() || ActiveCustomTimeStep.IsNone())
+		{
+			return LOCTEXT("LiveLinkTimeCode_CustomTimeStepError", "Connected clients will not be synchronized unless you select both a framerate and a subject for the Custom Time Step.");
+		}
+	}
+
+	if (bCachedIsTimecodeSource && !bCachedIsCustomTimeStepEnabled)
+	{
+		return LOCTEXT("LiveLinkTimeCode_TimecodeConnected", "Sending Timecode data to connected editors.");
+	}
+	else if (!bCachedIsTimecodeSource && bCachedIsCustomTimeStepEnabled)
 	{
-		return LOCTEXT("LiveLinkTimeCode_IsConnected", "Sending timecode data to connected editors.");
+		return LOCTEXT("LiveLinkTimeCode_CustomTimeStepConnected", "Sending CustomTimestep data to connected editors.");
 	}
-	return LOCTEXT("LiveLinkTimeCode_NotConnected", "No timecode data shared with connected editors.");
+	else if (bCachedIsTimecodeSource && bCachedIsCustomTimeStepEnabled)
+	{
+		return LOCTEXT("LiveLinkTimeCode_TimecodeAndCustomTimeStepConnected", "Sending Timecode and CustomTimeStep data to connected editors.");
+	}
+
+	return LOCTEXT("LiveLinkTimeCode_NotConnected", "No Timecode or CustomTimeStep data shared with connected editors.");
 }
 
 void SLiveLinkTimecode::OnEnableTimecodeToggled()
@@ -138,6 +237,45 @@ void SLiveLinkTimecode::OnEnableTimecodeToggled()
 	}
 }
 
+void SLiveLinkTimecode::OnEnableCustomTimeStepToggled()
+{
+	bCachedIsCustomTimeStepEnabled = !bCachedIsCustomTimeStepEnabled;
+
+	if (TSharedPtr<ILiveLinkHubSessionManager> SessionManager = FModuleManager::Get().GetModuleChecked<FLiveLinkHubModule>("LiveLinkHub").GetSessionManager())
+	{
+		SessionManager->GetCurrentSession()->SetShouldUseLiveLinkHubAsCustomTimeStepSource(bCachedIsCustomTimeStepEnabled);
+	}
+
+	if (bCachedIsCustomTimeStepEnabled)
+	{
+		SendUpdatedCustomTimeStepToEditor(MakeCustomTimeStepSettings());
+	}
+	else
+	{
+		// If we're disabling LLH as a CustomTimeStep source, reset the custom time step on the engine.
+		const TSharedPtr<FLiveLinkHub> LiveLinkHub = FModuleManager::Get().GetModuleChecked<FLiveLinkHubModule>("LiveLinkHub").GetLiveLinkHub();
+		if (LiveLinkHub.IsValid())
+		{
+			if (TSharedPtr<FLiveLinkHubProvider> Provider = LiveLinkHub->GetLiveLinkProvider())
+			{
+				Provider->ResetCustomTimeStepSettings();
+			}
+		}
+	}
+}
+
+void SLiveLinkTimecode::SendUpdatedCustomTimeStepToEditor(const FLiveLinkHubCustomTimeStepSettings& CustomTimeStepSettings)
+{
+	const TSharedPtr<FLiveLinkHub> LiveLinkHub = FModuleManager::Get().GetModuleChecked<FLiveLinkHubModule>("LiveLinkHub").GetLiveLinkHub();
+	if (LiveLinkHub.IsValid())
+	{
+		if (TSharedPtr<FLiveLinkHubProvider> Provider = LiveLinkHub->GetLiveLinkProvider())
+		{
+			Provider->UpdateCustomTimeStepSettings(CustomTimeStepSettings);
+		}
+	}
+}
+
 FLiveLinkHubTimecodeSettings SLiveLinkTimecode::MakeTimecodeSettings() const
 {
 	FLiveLinkHubTimecodeSettings Settings;
@@ -172,10 +310,30 @@ FLiveLinkHubTimecodeSettings SLiveLinkTimecode::MakeTimecodeSettings() const
 	return Settings;
 }
 
+FLiveLinkHubCustomTimeStepSettings SLiveLinkTimecode::MakeCustomTimeStepSettings() const
+{
+	FLiveLinkHubCustomTimeStepSettings Settings;
+	Settings.SubjectName = ActiveSubjectCustomTimeStep;
+
+	using namespace UE::LiveLinkTimecode::Private;
+	for (const TPair<FName, FTimecodeMenuItem>& Source : StaticCustomTimeStepMenu)
+	{
+		if (ActiveCustomTimeStep == Source.Key)
+		{
+			Settings.CustomTimeStepRate = Source.Value.Rate;
+			break;
+		}
+	};
+
+	Settings.FrameRateDivider = CachedFrameRateDivider;
+	Settings.bLockStepMode = bCachedLockStepMode;
+
+	return Settings;
+}
 
 void SLiveLinkTimecode::SendUpdatedTimecodeToEditor(const FLiveLinkHubTimecodeSettings& TimecodeSettings)
 {
-	const TSharedPtr<FLiveLinkHub> LiveLinkHub = FModuleManager::Get().GetModuleChecked<FLiveLinkHubModule>("LiveLinkHub").GetLiveLinkHub();
+	const TSharedPtr<FLiveLinkHub> LiveLinkHub = FModuleManager::Get().GetModuleChecked<FLiveLinkHubModule>("LiveLinkHub").GetLiveLinkHub();	
 	if (LiveLinkHub.IsValid())
 	{
 		if (TSharedPtr<FLiveLinkHubProvider> Provider = LiveLinkHub->GetLiveLinkProvider())
@@ -195,6 +353,93 @@ void SLiveLinkTimecode::UpdateTimecodeSettingsFromSession(const TSharedRef<ILive
 	ActiveTimecodeSource = Settings.SubjectName.IsNone() ? System24fps : Settings.SubjectName;
 }
 
+void SLiveLinkTimecode::UpdateCustomTimeStepSettingsFromSession(const TSharedRef<ILiveLinkHubSession>& InSession)
+{
+	using namespace UE::LiveLinkTimecode::Private;
+
+	const FLiveLinkHubCustomTimeStepSettings Settings = InSession->GetCustomTimeStepSettings();
+	bCachedIsCustomTimeStepEnabled = InSession->ShouldUseLiveLinkHubAsCustomTimeStepSource();
+	bCachedLockStepMode = Settings.bLockStepMode;
+
+	const TPair<FName, FTimecodeMenuItem>* TargetItem = Algo::FindByPredicate(StaticCustomTimeStepMenu, [Rate = Settings.CustomTimeStepRate](const TPair<FName, FTimecodeMenuItem>& Source) { return Rate == Source.Value.Rate;  });
+	if (TargetItem)
+	{
+		ActiveCustomTimeStep = TargetItem->Key;
+	}
+
+	CachedFrameRateDivider = Settings.FrameRateDivider;
+	ActiveSubjectCustomTimeStep = Settings.SubjectName;
+}
+
+void SLiveLinkTimecode::SetCustomTimeStepSubject(const FName SubjectName)
+{
+	if (SubjectName == ActiveSubjectCustomTimeStep)
+	{
+		return;
+	}
+
+	ActiveSubjectCustomTimeStep = SubjectName;
+
+	const FLiveLinkHubCustomTimeStepSettings Settings = MakeCustomTimeStepSettings();
+
+	const TSharedPtr<FLiveLinkHub> LiveLinkHub = FModuleManager::Get().GetModuleChecked<FLiveLinkHubModule>("LiveLinkHub").GetLiveLinkHub();
+	if (LiveLinkHub.IsValid())
+	{
+		LiveLinkHub->GetSessionManager()->GetCurrentSession()->SetCustomTimeStepSettings(Settings);
+	}
+
+	if (bCachedIsCustomTimeStepEnabled)
+	{
+		SendUpdatedCustomTimeStepToEditor(Settings);
+	}
+}
+
+void SLiveLinkTimecode::SetCustomTimeStepRateDivider(uint32 RateDivider)
+{
+	if (RateDivider == CachedFrameRateDivider)
+	{
+		return;
+	}
+
+	CachedFrameRateDivider = RateDivider;
+
+	const FLiveLinkHubCustomTimeStepSettings Settings = MakeCustomTimeStepSettings();
+	
+	const TSharedPtr<FLiveLinkHub> LiveLinkHub = FModuleManager::Get().GetModuleChecked<FLiveLinkHubModule>("LiveLinkHub").GetLiveLinkHub();
+	if (LiveLinkHub.IsValid())
+	{
+		LiveLinkHub->GetSessionManager()->GetCurrentSession()->SetCustomTimeStepSettings(Settings);
+	}
+
+	if (bCachedIsCustomTimeStepEnabled)
+	{
+		SendUpdatedCustomTimeStepToEditor(Settings);
+	}
+}
+
+void SLiveLinkTimecode::SetCustomTimeStep(const FName SourceId)
+{
+	if (SourceId == ActiveCustomTimeStep)
+	{
+		return;
+	}
+
+	ActiveCustomTimeStep = SourceId;
+
+	const FLiveLinkHubCustomTimeStepSettings Settings = MakeCustomTimeStepSettings();
+
+	const TSharedPtr<FLiveLinkHub> LiveLinkHub = FModuleManager::Get().GetModuleChecked<FLiveLinkHubModule>("LiveLinkHub").GetLiveLinkHub();
+	if (LiveLinkHub.IsValid())
+	{
+		LiveLinkHub->GetSessionManager()->GetCurrentSession()->SetCustomTimeStepSettings(Settings);
+	}
+
+	if (bCachedIsCustomTimeStepEnabled)
+	{
+		SendUpdatedCustomTimeStepToEditor(Settings);
+	}
+}
+
 void SLiveLinkTimecode::SetTimecodeSource(const FName SourceId)
 {
 	if (SourceId == ActiveTimecodeSource)
@@ -219,34 +464,55 @@ void SLiveLinkTimecode::SetTimecodeSource(const FName SourceId)
 	}
 }
 
-TSharedRef<SWidget> SLiveLinkTimecode::MakeMenu()
+TSharedRef<SWidget> SLiveLinkTimecode::MakeFrameRateDivider()
 {
 	FMenuBuilder MenuBuilder(/*bInShouldCloseWindowAfterMenuSelection=*/true, nullptr);
 
-	MenuBuilder.AddMenuEntry(
-		LOCTEXT("LiveLinkHubTimecodeSource", "Enable Timecode Source"),
-		LOCTEXT("LiveLinkHubTimecodeSource_Tooltip", "Make this Live Link Hub a time code source for connected editors."),
-		FSlateIcon(),
-		FUIAction(
-			FExecuteAction::CreateSP(this, &SLiveLinkTimecode::OnEnableTimecodeToggled),
-			FCanExecuteAction::CreateLambda([] { return true; }),
-			FIsActionChecked::CreateLambda([this] { return bCachedIsTimecodeSource; })),
-		NAME_None,
-		EUserInterfaceActionType::ToggleButton);
-
+	return SNew(SHorizontalBox)
+		+ SHorizontalBox::Slot()
+		.FillWidth(1.f)
+		[
+			SNew(SSpacer)
+		]
+		+ SHorizontalBox::Slot()
+		.AutoWidth()
+		.Padding(0.f, 0.f, 15.f, 0.f)
+		[
+			SNew(SBox)
+				.WidthOverride(50.f)
+				[
+					SNew(SNumericEntryBox<uint32>)
+						.AllowSpin(true)
+						.MinDesiredValueWidth(10)
+						.MinValue(1)
+						.MinSliderValue(1)
+						.MaxSliderValue(100)
+						.OnValueChanged_Lambda([this](uint32 NewValue)
+							{
+								SetCustomTimeStepRateDivider(NewValue);
+							})
+						.Value_Lambda([this]()
+							{
+								return CachedFrameRateDivider;
+							})
+				]
+		];
+}
 
+void SLiveLinkTimecode::MakeTimecodeMenuItems(FMenuBuilder& MenuBuilder)
+{
 	MenuBuilder.BeginSection("LiveLinkHub.Timecode.TimecodeProvider", LOCTEXT("TimecodeProviderSection", "Timecode Provider"));
 
 	auto GenerateUIAction = [this](const FName Id) -> FUIAction
-	{
-		// Generate a UI Action
-		return FUIAction(
-			FExecuteAction::CreateSP(this, &SLiveLinkTimecode::SetTimecodeSource, Id),
-			FCanExecuteAction::CreateLambda([] { return true; }),
-			FIsActionChecked::CreateLambda([this, Id] { return Id == ActiveTimecodeSource; }));
-	};
+		{
+			// Generate a UI Action
+			return FUIAction(
+				FExecuteAction::CreateSP(this, &SLiveLinkTimecode::SetTimecodeSource, Id),
+				FCanExecuteAction::CreateLambda([] { return true; }),
+				FIsActionChecked::CreateLambda([this, Id] { return Id == ActiveTimecodeSource; }));
+		};
 	using namespace UE::LiveLinkTimecode::Private;
-	for (const TPair<FName,FTimecodeMenuItem>& Source : StaticTimecodeMenu)
+	for (const TPair<FName, FTimecodeMenuItem>& Source : StaticTimecodeMenu)
 	{
 		MenuBuilder.AddMenuEntry(
 			Source.Value.Label,
@@ -272,6 +538,143 @@ TSharedRef<SWidget> SLiveLinkTimecode::MakeMenu()
 			EUserInterfaceActionType::Check);
 	}
 	MenuBuilder.EndSection();
+}
+
+
+void SLiveLinkTimecode::MakeCustomTimeStepMenuItems(FMenuBuilder& MenuBuilder)
+{
+	auto GenerateUIAction = [this](const FName Id) -> FUIAction
+		{
+			// Generate a UI Action
+			return FUIAction(
+				FExecuteAction::CreateSP(this, &SLiveLinkTimecode::SetCustomTimeStep, Id),
+				FCanExecuteAction::CreateLambda([] { return true; }),
+				FIsActionChecked::CreateLambda([this, Id] { return Id == ActiveCustomTimeStep; }));
+		};
+	using namespace UE::LiveLinkTimecode::Private;
+	for (const TPair<FName, FTimecodeMenuItem>& Source : StaticCustomTimeStepMenu)
+	{
+		MenuBuilder.AddMenuEntry(
+			Source.Value.Label,
+			Source.Value.ToolTip,
+			FSlateIcon(),
+			GenerateUIAction(Source.Key),
+			NAME_None,
+			EUserInterfaceActionType::Check);
+	}
+
+	MenuBuilder.AddMenuSeparator();
+
+	MenuBuilder.AddMenuEntry(
+		LOCTEXT("LiveLinkCustomTimeStepStep", "Lock Step Mode"),
+		LOCTEXT("LiveLinkCustomTimeStepUseTimecode_Tooltip", "Enable / disable lock step mode."),
+		FSlateIcon(),
+		FUIAction(FExecuteAction::CreateLambda(
+			[this]() mutable
+			{
+				bCachedLockStepMode = !bCachedLockStepMode;
+			}),
+			FCanExecuteAction::CreateLambda([] { return true; }),
+			FIsActionChecked::CreateLambda([this] { return bCachedLockStepMode; })),
+		NAME_None,
+		EUserInterfaceActionType::Check);
+
+	MenuBuilder.AddWidget(MakeFrameRateDivider(), LOCTEXT("FrameRateDivider", "Rate Divider"), false);
+
+	FNewMenuDelegate SubMenuDelegate = FNewMenuDelegate::CreateLambda(
+		[this](FMenuBuilder& SubMenuBuilder)
+		{
+			TArray<FName> Subjects;
+			WorkingClient->GetSubjectNames(Subjects);
+			for (FName Subject : Subjects)
+			{
+				FTextBuilder SubjectText;
+				SubjectText.AppendLineFormat(
+					LOCTEXT("LiveLinkHubCustomTimeStepSourceSubject_Tooltip", "Use {0} as a frame lock source."),
+					FText::FromName(Subject)
+				);
+
+				FUIAction Action = FUIAction(
+					FExecuteAction::CreateSP(this, &SLiveLinkTimecode::SetCustomTimeStepSubject, Subject),
+					FCanExecuteAction::CreateLambda([] { return true; }),
+					FIsActionChecked::CreateLambda([this, Subject] { return Subject == ActiveSubjectCustomTimeStep; }));
+
+				SubMenuBuilder.AddMenuEntry(
+					FText::FromName(Subject),
+					SubjectText.ToText(),
+					FSlateIcon(),
+					Action,
+					NAME_None,
+					EUserInterfaceActionType::Check);
+			}
+		});
+
+	MenuBuilder.AddSubMenu(
+		LOCTEXT("LiveLinkHubCustomTimeStep", "Subject Key"),
+		LOCTEXT("LiveLinkHubTimestep_Tooltip", "The specific subject we listen."),
+		SubMenuDelegate,
+		FUIAction(),
+		NAME_None,
+		EUserInterfaceActionType::None);
+}
+
+void SLiveLinkTimecode::OnLockStepModeToggled()
+{
+	bCachedLockStepMode = !bCachedLockStepMode;
+
+	const FLiveLinkHubCustomTimeStepSettings Settings = MakeCustomTimeStepSettings();
+
+	const TSharedPtr<FLiveLinkHub> LiveLinkHub = FModuleManager::Get().GetModuleChecked<FLiveLinkHubModule>("LiveLinkHub").GetLiveLinkHub();
+	if (LiveLinkHub.IsValid())
+	{
+		LiveLinkHub->GetSessionManager()->GetCurrentSession()->SetCustomTimeStepSettings(Settings);
+	}
+
+	if (bCachedIsCustomTimeStepEnabled)
+	{
+		SendUpdatedCustomTimeStepToEditor(Settings);
+	}
+}
+
+TSharedRef<SWidget> SLiveLinkTimecode::MakeMenu()
+{
+	FMenuBuilder MenuBuilder(/*bInShouldCloseWindowAfterMenuSelection=*/true, nullptr);
+
+	FNewMenuDelegate TimecodeSubMenuDelegate = FNewMenuDelegate::CreateLambda(
+		[this](FMenuBuilder& SubMenuBuilder)
+		{
+			MakeTimecodeMenuItems(SubMenuBuilder);
+		});
+
+	MenuBuilder.AddSubMenu(
+		LOCTEXT("LiveLinkHubTimecodeSource", "Enable Timecode Source"),
+		LOCTEXT("LiveLinkHubTimecodeSource_Tooltip", "Make this Live Link Hub a time code source for connected editors."),
+		TimecodeSubMenuDelegate,
+		FUIAction(
+			FExecuteAction::CreateSP(this, &SLiveLinkTimecode::OnEnableTimecodeToggled),
+			FCanExecuteAction::CreateLambda([] { return true; }),
+			FIsActionChecked::CreateLambda([this] { return bCachedIsTimecodeSource; })),
+		NAME_None,
+		EUserInterfaceActionType::ToggleButton);
+
+
+	FNewMenuDelegate SubMenuDelegate = FNewMenuDelegate::CreateLambda(
+		[this](FMenuBuilder& SubMenuBuilder)
+		{
+			MakeCustomTimeStepMenuItems(SubMenuBuilder);
+		});
+
+	MenuBuilder.AddSubMenu(
+		LOCTEXT("LiveLinkHubCustomTimeStep", "Enable Custom Time Step"),
+		LOCTEXT("LiveLinkHubTimestep_Tooltip", "Enable CustomTimeStep forces the framerate of connected clients to the choosen value.."),
+		SubMenuDelegate,
+		FUIAction(
+			FExecuteAction::CreateSP(this, &SLiveLinkTimecode::OnEnableCustomTimeStepToggled),
+			FCanExecuteAction::CreateLambda([] { return true; }),
+			FIsActionChecked::CreateLambda([this] { return bCachedIsCustomTimeStepEnabled; })),
+		NAME_None,
+		EUserInterfaceActionType::ToggleButton);
+
 	return MenuBuilder.MakeWidget();
 }
 
@@ -284,11 +687,13 @@ void SLiveLinkTimecode::Construct(const FArguments& InArgs)
 	if (SessionManager.IsValid())
 	{
 		UpdateTimecodeSettingsFromSession(SessionManager->GetCurrentSession().ToSharedRef());
+		UpdateCustomTimeStepSettingsFromSession(SessionManager->GetCurrentSession().ToSharedRef());
 
 		SessionManager->OnActiveSessionChanged().AddLambda([this](const TSharedRef<ILiveLinkHubSession>& InSession)
 		{
 			// Update the UI when a config is loaded.
-				UpdateTimecodeSettingsFromSession(InSession);
+			UpdateTimecodeSettingsFromSession(InSession);
+			UpdateCustomTimeStepSettingsFromSession(InSession);
 		});
 	}
 
@@ -312,7 +717,7 @@ void SLiveLinkTimecode::Construct(const FArguments& InArgs)
 			[
 				SNew(STextBlock)
 					.Font(FAppStyle::Get().GetFontStyle("FontAwesome.8"))
-					.ColorAndOpacity(this, &SLiveLinkTimecode::GetTimecodeStatusColor)
+					.ColorAndOpacity(this, &SLiveLinkTimecode::GetStatusColor)
 					.Text(FEditorFontGlyphs::Circle)
 			]
 			+ SHorizontalBox::Slot()
diff --git a/Engine/Plugins/Animation/LiveLinkHub/Source/LiveLinkHub/Private/UI/Widgets/SLiveLinkTimecode.h b/Engine/Plugins/Animation/LiveLinkHub/Source/LiveLinkHub/Private/UI/Widgets/SLiveLinkTimecode.h
--- a/Engine/Plugins/Animation/LiveLinkHub/Source/LiveLinkHub/Private/UI/Widgets/SLiveLinkTimecode.h
+++ b/Engine/Plugins/Animation/LiveLinkHub/Source/LiveLinkHub/Private/UI/Widgets/SLiveLinkTimecode.h
@@ -7,6 +7,8 @@
 class FName;
 class FLiveLinkClient;
 struct FLiveLinkHubTimecodeSettings;
+struct FLiveLinkHubCustomTimeStepSettings;
+class FMenuBuilder;
 
 /**
  * Utility class to display timecode status on the menu bar. Provides user with
@@ -28,34 +30,83 @@ public:
 
 private:
 	/** Indicate the status of the time code send status */
-	FSlateColor GetTimecodeStatusColor() const;
+	FSlateColor GetStatusColor() const;
 
 	/** Gets the tooltip text for the time code widget */
 	FText GetTimecodeTooltip() const;
 
+	/** Make the menu entry that corresponds to the FrameRateDivider property on the CustomTimeStep */
+	TSharedRef<SWidget> MakeFrameRateDivider();
+
+	/** Make menu items for the frame lock submenu */
+	void MakeCustomTimeStepMenuItems(FMenuBuilder& MenuBuilder);
+
+	/** Make the menu items for the timecode provider menu. */
+	void MakeTimecodeMenuItems(FMenuBuilder& MenuBuilder);
+
 	/** Makes the combo button menu for Timecode source */
 	TSharedRef<SWidget> MakeMenu();
 
 	/** Invoked by slate when the time code setting is enabled. */
 	void OnEnableTimecodeToggled();
 
+	/** Invoked by slate when the frame lock setting is enabled. */
+	void OnEnableCustomTimeStepToggled();
+
+	/** Invoked by slate when the lock step mode setting is toggle. */
+	void OnLockStepModeToggled();
+
+	/** Set the new frame rate divider.*/
+	void SetCustomTimeStepRateDivider(uint32 Divider);
+
+	/** Set the subject name to use with frame lock setting.*/
+	void SetCustomTimeStepSubject(const FName SubjectName);
+
 	/** Specifies the source for timecode */
 	void SetTimecodeSource(const FName SourceId);
 
+	/** Specifies the source for CustomTimeStep */
+	void SetCustomTimeStep(const FName SourceId);
+
 	/** Sends the current timecode settings to the connected editor(s). */
 	void SendUpdatedTimecodeToEditor(const FLiveLinkHubTimecodeSettings& TimecodeSettings);
 
+	/** Create a timecode settings objects from the session config. */
 	FLiveLinkHubTimecodeSettings MakeTimecodeSettings() const;
 
+	/** Create a CustomTimeStep settings objects from the session config. */
+	FLiveLinkHubCustomTimeStepSettings MakeCustomTimeStepSettings() const;
+
+	/** Sends the current frame lock settings to the connected editor(s). */
+	void SendUpdatedCustomTimeStepToEditor(const FLiveLinkHubCustomTimeStepSettings& CustomTimeStepSettings);
+
 	/** Update the local timecode properties from the timecode settings. */
 	void UpdateTimecodeSettingsFromSession(const TSharedRef<class ILiveLinkHubSession>& InSession);
+
+	/** Update the local CustomTimeStep properties from the session CustomTimeStep settings. */
+	void UpdateCustomTimeStepSettingsFromSession(const TSharedRef<ILiveLinkHubSession>& InSession);
 	
 	/** Do we have timecode source enabled */
 	bool bCachedIsTimecodeSource = false;
 
+	/** Do we have frame lock enabled. */
+	bool bCachedIsCustomTimeStepEnabled = false;
+
 	/** Name of the active timecode source. */
 	FName ActiveTimecodeSource;
 
+	/** Name of the active custom time step framerate. */
+	FName ActiveCustomTimeStep;
+
+	/** Enable / disable lock step mode for frame lock */
+	bool bCachedLockStepMode = true;
+
+	/** Frame rate divider for frame lock mode. */
+	uint32 CachedFrameRateDivider = 1;
+
+	/** The name of the subject applied to the frame lock on Live Link Custom Timestep. */
+	FName ActiveSubjectCustomTimeStep;
+
 	/** Cached value for the FLiveLinkClient */
 	FLiveLinkClient* WorkingClient = nullptr;
 };
diff --git a/Engine/Plugins/Animation/LiveLinkHub/Source/LiveLinkHubMessaging/LiveLinkHubMessaging.build.cs b/Engine/Plugins/Animation/LiveLinkHub/Source/LiveLinkHubMessaging/LiveLinkHubMessaging.build.cs
--- a/Engine/Plugins/Animation/LiveLinkHub/Source/LiveLinkHubMessaging/LiveLinkHubMessaging.build.cs
+++ b/Engine/Plugins/Animation/LiveLinkHub/Source/LiveLinkHubMessaging/LiveLinkHubMessaging.build.cs
@@ -9,6 +9,7 @@ namespace UnrealBuildTool.Rules
 			PublicDependencyModuleNames.AddRange(
 			new string[]
 			{
+				"TimeManagement",
 				"LiveLinkMessageBusFramework"
 			});
 
diff --git a/Engine/Plugins/Animation/LiveLinkHub/Source/LiveLinkHubMessaging/Private/LiveLinkHubMessageBusSource.cpp b/Engine/Plugins/Animation/LiveLinkHub/Source/LiveLinkHubMessaging/Private/LiveLinkHubMessageBusSource.cpp
--- a/Engine/Plugins/Animation/LiveLinkHub/Source/LiveLinkHubMessaging/Private/LiveLinkHubMessageBusSource.cpp
+++ b/Engine/Plugins/Animation/LiveLinkHub/Source/LiveLinkHubMessaging/Private/LiveLinkHubMessageBusSource.cpp
@@ -92,7 +92,16 @@ void FLiveLinkHubMessageBusSource::InitializeMessageEndpoint(FMessageEndpointBui
 {
 	FLiveLinkMessageBusSource::InitializeMessageEndpoint(EndpointBuilder);
 	EndpointBuilder
-		.Handling<FLiveLinkHubTimecodeSettings>(this, &FLiveLinkHubMessageBusSource::HandleTimecodeSettings);
+		.Handling<FLiveLinkHubTimecodeSettings>(this, &FLiveLinkHubMessageBusSource::HandleTimecodeSettings)
+		.Handling<FLiveLinkHubCustomTimeStepSettings>(this, &FLiveLinkHubMessageBusSource::HandleCustomTimeStepSettings);
+}
+
+
+void FLiveLinkHubMessageBusSource::HandleCustomTimeStepSettings(const FLiveLinkHubCustomTimeStepSettings& Message, const TSharedRef<IMessageContext, ESPMode::ThreadSafe>& Context)
+{
+	ExecuteOnGameThread(UE_SOURCE_LOCATION, [Message]() {
+		Message.AssignCustomTimeStepToEngine();
+	});
 }
 
 void FLiveLinkHubMessageBusSource::HandleTimecodeSettings(const FLiveLinkHubTimecodeSettings& Message, const TSharedRef<IMessageContext, ESPMode::ThreadSafe>& Context)
diff --git a/Engine/Plugins/Animation/LiveLinkHub/Source/LiveLinkHubMessaging/Private/LiveLinkHubMessageBusSource.h b/Engine/Plugins/Animation/LiveLinkHub/Source/LiveLinkHubMessaging/Private/LiveLinkHubMessageBusSource.h
--- a/Engine/Plugins/Animation/LiveLinkHub/Source/LiveLinkHubMessaging/Private/LiveLinkHubMessageBusSource.h
+++ b/Engine/Plugins/Animation/LiveLinkHub/Source/LiveLinkHubMessaging/Private/LiveLinkHubMessageBusSource.h
@@ -2,6 +2,7 @@
 
 #pragma once
 
+#include "LiveLinkHubMessages.h"
 #include "LiveLinkMessageBusSource.h"
 
 class UWorld;
@@ -26,7 +27,11 @@ protected:
 	//~ End FLiveLinkMessageBusSource
 
 private:
+	/** Handle a custom time step settings message and update the engine's custom time step settings accordingly. */
+	void HandleCustomTimeStepSettings(const FLiveLinkHubCustomTimeStepSettings& Message, const TSharedRef<IMessageContext, ESPMode::ThreadSafe>& Context);
+	/** Handle a timecode settings message and update the engine's timecode settings accordingly. */
 	void HandleTimecodeSettings(const FLiveLinkHubTimecodeSettings& Message, const TSharedRef<IMessageContext, ESPMode::ThreadSafe>& Context);
+
 	/** Send information about this UE client to the connected provider */
 	void SendClientInfoMessage();
 	/** Handler called on map changed to update the livelink hub. */
diff --git a/Engine/Plugins/Animation/LiveLinkHub/Source/LiveLinkHubMessaging/Private/LiveLinkHubMessages.cpp b/Engine/Plugins/Animation/LiveLinkHub/Source/LiveLinkHubMessaging/Private/LiveLinkHubMessages.cpp
--- a/Engine/Plugins/Animation/LiveLinkHub/Source/LiveLinkHubMessaging/Private/LiveLinkHubMessages.cpp
+++ b/Engine/Plugins/Animation/LiveLinkHub/Source/LiveLinkHubMessaging/Private/LiveLinkHubMessages.cpp
@@ -16,6 +16,67 @@ FName FLiveLinkHubMessageAnnotation::ProviderTypeAnnotation = TEXT("ProviderType
 
 DEFINE_LOG_CATEGORY_STATIC(LogLiveLinkHubMessages, Log, All);
 
+namespace UE::LiveLinkHubMessages::Private
+{
+	TOptional<FLiveLinkSubjectKey> FindSubjectKey(ILiveLinkClient* LiveLinkClient, FName SubjectName)
+	{
+		TArray<FLiveLinkSubjectKey> Subjects = LiveLinkClient->GetSubjects(true, true);
+		// We need to map the named subject to the list of subject keys available.
+		for (const FLiveLinkSubjectKey& Key : Subjects)
+		{
+			if (Key.SubjectName == SubjectName)
+			{
+				return Key;
+			}
+		}
+		return {};
+	}
+}
+
+void FLiveLinkHubCustomTimeStepSettings::AssignCustomTimeStepToEngine() const
+{
+	IModularFeatures& ModularFeatures = IModularFeatures::Get();
+	if (!ModularFeatures.IsModularFeatureAvailable(ILiveLinkClient::ModularFeatureName))
+	{
+		return;
+	}
+
+	if (GEngine == nullptr)
+	{
+		return;
+	}
+
+	if (bResetCustomTimeStep)
+	{
+		UEngineCustomTimeStep* CurrentCustomTimeStep = GEngine->GetCustomTimeStep();
+
+		if (Cast<ULiveLinkHubCustomTimeStep>(CurrentCustomTimeStep))
+		{
+			UE_LOG(LogLiveLinkHubMessages, Display, TEXT("CustomTimeStep reset event"));
+
+			// We only issue a timestep reset if we are a LiveLinkHubCustomTimeStep.  This way we don't reset any custom time step that the user
+			// may have set in the editor.
+			GEngine->Exec(GEngine->GetCurrentPlayWorld(nullptr), TEXT("CustomTimeStep.reset"));
+		}
+		return;
+	}
+	
+	UE_LOG(LogLiveLinkHubMessages, Display, TEXT("CustomTimeStep change event %s - %s"), *SubjectName.ToString(), *CustomTimeStepRate.ToPrettyText().ToString());
+
+	ILiveLinkClient* LiveLinkClient = &ModularFeatures.GetModularFeature<ILiveLinkClient>(ILiveLinkClient::ModularFeatureName);
+
+	ULiveLinkHubCustomTimeStep* NewCustomTimeStep = NewObject<ULiveLinkHubCustomTimeStep>(GEngine, ULiveLinkHubCustomTimeStep::StaticClass());
+	NewCustomTimeStep->LiveLinkDataRate = CustomTimeStepRate;
+
+	if (TOptional<FLiveLinkSubjectKey> Target = UE::LiveLinkHubMessages::Private::FindSubjectKey(LiveLinkClient, SubjectName))
+	{
+		NewCustomTimeStep->SubjectKey = *Target;
+	}
+
+	// Override the custom timestep for the engine.
+	GEngine->SetCustomTimeStep(NewCustomTimeStep);
+}
+
 void FLiveLinkHubTimecodeSettings::AssignTimecodeSettingsAsProviderToEngine() const
 {
 	IModularFeatures& ModularFeatures = IModularFeatures::Get();
@@ -39,16 +100,7 @@ void FLiveLinkHubTimecodeSettings::AssignTimecodeSettingsAsProviderToEngine() co
 	}
 	else if (Source == ELiveLinkHubTimecodeSource::UseSubjectName)
 	{
-		TArray<FLiveLinkSubjectKey> Subjects = LiveLinkClient->GetSubjects(true, true);
-		TOptional<FLiveLinkSubjectKey> Target;
-		// We need to map the named subject to the list of subject keys available.
-		for (const FLiveLinkSubjectKey& Key : Subjects)
-		{
-			if (Key.SubjectName == SubjectName)
-			{
-				Target = Key;
-			}
-		}
+		TOptional<FLiveLinkSubjectKey> Target = UE::LiveLinkHubMessages::Private::FindSubjectKey(LiveLinkClient, SubjectName);
 
 		if (Target)
 		{
diff --git a/Engine/Plugins/Animation/LiveLinkHub/Source/LiveLinkHubMessaging/Public/LiveLinkHubMessages.h b/Engine/Plugins/Animation/LiveLinkHub/Source/LiveLinkHubMessaging/Public/LiveLinkHubMessages.h
--- a/Engine/Plugins/Animation/LiveLinkHub/Source/LiveLinkHubMessaging/Public/LiveLinkHubMessages.h
+++ b/Engine/Plugins/Animation/LiveLinkHub/Source/LiveLinkHubMessaging/Public/LiveLinkHubMessages.h
@@ -2,10 +2,13 @@
 
 #pragma once
 
+#include "GenlockedCustomTimeStep.h"
 #include "LiveLinkMessages.h"
 #include "LiveLinkTypes.h"
 #include "Misc/FrameRate.h"
 
+#include "LiveLinkCustomTimeStep.h"
+
 #include "LiveLinkHubMessages.generated.h"
 
 /** Annotation put on MessageBus messages to indicate the type of provider used. 
@@ -22,6 +25,40 @@ namespace UE::LiveLinkHub::Private
 	extern const LIVELINKHUBMESSAGING_API FName LiveLinkHubProviderType;
 }
 
+UCLASS(Hidden)
+class ULiveLinkHubCustomTimeStep : public ULiveLinkCustomTimeStep
+{
+	GENERATED_BODY()
+};
+
+USTRUCT()
+struct LIVELINKHUBMESSAGING_API FLiveLinkHubCustomTimeStepSettings
+{
+	GENERATED_BODY()
+
+	/** If this is true, the engine custom time step will be reset. */
+	UPROPERTY()
+	bool bResetCustomTimeStep = false;
+
+	/** Corresponds to the lock step mode in ULiveLinkCustomTimeStep. */
+	UPROPERTY()
+	bool bLockStepMode = true;
+
+	/** Corresponds to the frame rate divider in ULiveLinkCustomTimeStep */
+	UPROPERTY()
+	uint32 FrameRateDivider = 1;
+
+	/** If we are locking the editor frame rate to the subject then this property holds that subject name. */
+	UPROPERTY()
+	FName SubjectName;
+
+	/** Desired frame rate to lock the editor. This corresponds to the LiveLinkDataRate in ULiveLinkCustomTimeStep */
+	UPROPERTY()
+	FFrameRate CustomTimeStepRate;
+
+	/** Assign the frame lock settings to the engine. */
+	void AssignCustomTimeStepToEngine() const;
+};
 
 UENUM()
 enum class ELiveLinkHubTimecodeSource
