Fix CIS warnings caused by duplicate loctext macros in SLiveLinkTimecode

diff --git a/Engine/Plugins/Animation/LiveLinkHub/Source/LiveLinkHub/Private/UI/Widgets/SLiveLinkTimecode.cpp b/Engine/Plugins/Animation/LiveLinkHub/Source/LiveLinkHub/Private/UI/Widgets/SLiveLinkTimecode.cpp
--- a/Engine/Plugins/Animation/LiveLinkHub/Source/LiveLinkHub/Private/UI/Widgets/SLiveLinkTimecode.cpp
+++ b/Engine/Plugins/Animation/LiveLinkHub/Source/LiveLinkHub/Private/UI/Widgets/SLiveLinkTimecode.cpp
@@ -158,8 +158,8 @@ namespace UE::LiveLinkTimecode::Private
 				LOCTEXT("LiveLinkHubCustomTimeStep100fps_Tooltip", "Use a 100 FPS timecode for CustomTimeStep.") }               },
 			{
              CustomTimeStep120fps,    { FFrameRate(120, 1),
-				  LOCTEXT("LiveLinkHubCustomTimeStep100fps", "120 fps"),
-				  LOCTEXT("LiveLinkHubCustomTimeStep100fps_Tooltip", "Use a 120 FPS timecode for CustomTimeStep.") }
+				  LOCTEXT("LiveLinkHubCustomTimeStep120fps", "120 fps"),
+				  LOCTEXT("LiveLinkHubCustomTimeStep120fps_Tooltip", "Use a 120 FPS timecode for CustomTimeStep.") }
 		}};
 		return Definition;
 	}();
@@ -610,8 +610,8 @@ void SLiveLinkTimecode::MakeCustomTimeStepMenuItems(FMenuBuilder& MenuBuilder)
 		});
 
 	MenuBuilder.AddSubMenu(
-		LOCTEXT("LiveLinkHubCustomTimeStep", "Subject Key"),
-		LOCTEXT("LiveLinkHubTimestep_Tooltip", "The specific subject we listen."),
+		LOCTEXT("LiveLinkHubCustomTimeStepSubjectKey", "Subject Key"),
+		LOCTEXT("LiveLinkHubTimestepSubjectKey_Tooltip", "The subject that's used for synchronizing the engine."),
 		SubMenuDelegate,
 		FUIAction(),
 		NAME_None,
