Fix crash when exiting LiveLinkHub or launching commandlet applications  (introduced in 37530644)

diff --git a/Engine/Plugins/Animation/LiveLinkHub/Source/LiveLinkHub/Private/LiveLinkHub.cpp b/Engine/Plugins/Animation/LiveLinkHub/Source/LiveLinkHub/Private/LiveLinkHub.cpp
--- a/Engine/Plugins/Animation/LiveLinkHub/Source/LiveLinkHub/Private/LiveLinkHub.cpp
+++ b/Engine/Plugins/Animation/LiveLinkHub/Source/LiveLinkHub/Private/LiveLinkHub.cpp
@@ -105,6 +105,9 @@ void FLiveLinkHub::InitializePostEngineInit()
 
 FLiveLinkHub::~FLiveLinkHub()
 {
+	IModularFeatures::Get().UnregisterModularFeature(IConsoleCommandExecutor::ModularFeatureName(), CommandExecutor.Get());
+	CommandExecutor.Reset();
+
 	UnregisterLiveLinkHubSettings();
 
 	RecordingController.Reset();
diff --git a/Engine/Plugins/Animation/LiveLinkHub/Source/LiveLinkHub/Private/LiveLinkHubModule.cpp b/Engine/Plugins/Animation/LiveLinkHub/Source/LiveLinkHub/Private/LiveLinkHubModule.cpp
--- a/Engine/Plugins/Animation/LiveLinkHub/Source/LiveLinkHub/Private/LiveLinkHubModule.cpp
+++ b/Engine/Plugins/Animation/LiveLinkHub/Source/LiveLinkHub/Private/LiveLinkHubModule.cpp
@@ -23,13 +23,19 @@ void FLiveLinkHubModule::StartupModule()
 {
 	TRACE_CPUPROFILER_EVENT_SCOPE(StartLiveLinkHub);
 
-	if (!GetDefault<ULiveLinkHubSettings>()->bTickOnGameThread)
+	bCreateLiveLinkHubInstance =
+		GConfig->GetBoolOrDefault(TEXT("LiveLink"), TEXT("bCreateLiveLinkHubInstance"), false, GEngineIni);
+
+	if (bCreateLiveLinkHubInstance)
 	{
-		Ticker.StartTick();
-	}
+		if (!GetDefault<ULiveLinkHubSettings>()->bTickOnGameThread)
+		{
+			Ticker.StartTick();
+		}
 
-	LiveLinkHub = MakeShared<FLiveLinkHub>();
-	LiveLinkHub->Initialize(Ticker);
+		LiveLinkHub = MakeShared<FLiveLinkHub>();
+		LiveLinkHub->Initialize(Ticker);
+	}
 
 	FPropertyEditorModule& PropertyModule = FModuleManager::LoadModuleChecked<FPropertyEditorModule>("PropertyEditor");
 
@@ -67,13 +73,13 @@ void FLiveLinkHubModule::ShutdownModule()
 		}
 	}
 
-	if (!GetDefault<ULiveLinkHubSettings>()->bTickOnGameThread)
+	if (bCreateLiveLinkHubInstance)
 	{
-		Ticker.Exit();
 		Ticker.Stop();
-	}
+		Ticker.Exit();
 
-	LiveLinkHub.Reset();
+		LiveLinkHub.Reset();
+	}
 }
 
 TSharedPtr<FLiveLinkHub> FLiveLinkHubModule::GetLiveLinkHub() const
diff --git a/Engine/Plugins/Animation/LiveLinkHub/Source/LiveLinkHub/Private/LiveLinkHubModule.h b/Engine/Plugins/Animation/LiveLinkHub/Source/LiveLinkHub/Private/LiveLinkHubModule.h
--- a/Engine/Plugins/Animation/LiveLinkHub/Source/LiveLinkHub/Private/LiveLinkHubModule.h
+++ b/Engine/Plugins/Animation/LiveLinkHub/Source/LiveLinkHub/Private/LiveLinkHubModule.h
@@ -50,4 +50,7 @@ private:
 	
 	/** Overrides standard livelink settings customization. Config-driven. True in LiveLink Hub standalone application, false in UE. */
 	bool bUseSettingsDetailCustomization = false;
+
+	/** Whether or not to create a LiveLinkHub instance (false in editor). */
+	bool bCreateLiveLinkHubInstance = false;
 };
diff --git a/Engine/Source/Programs/LiveLinkHub/Config/DefaultEngine.ini b/Engine/Source/Programs/LiveLinkHub/Config/DefaultEngine.ini
--- a/Engine/Source/Programs/LiveLinkHub/Config/DefaultEngine.ini
+++ b/Engine/Source/Programs/LiveLinkHub/Config/DefaultEngine.ini
@@ -4,6 +4,7 @@ bUseLiveLinkHubSubjectSettingsDetailsCustomization=true
 bPanelControllerSeparateSourcesSubjects=true
 bUseModularClientReference=true
 bEnableParentSubjects=true
+bCreateLiveLinkHubInstance=true
 
 [LiveLinkHub]
 bEnableLLHMessageBusSourceFactory=false
