CIS Fix - Ensure we don't initialize LiveLinkHub in commandlets to avoid creating a LiveLinkHub window in DDC Build jobs

diff --git a/Engine/Plugins/Animation/LiveLinkHub/Source/LiveLinkHub/Private/LiveLinkHubModule.cpp b/Engine/Plugins/Animation/LiveLinkHub/Source/LiveLinkHub/Private/LiveLinkHubModule.cpp
--- a/Engine/Plugins/Animation/LiveLinkHub/Source/LiveLinkHub/Private/LiveLinkHubModule.cpp
+++ b/Engine/Plugins/Animation/LiveLinkHub/Source/LiveLinkHub/Private/LiveLinkHubModule.cpp
@@ -2,9 +2,12 @@
 
 #include "LiveLinkHubModule.h"
 
+#include "CoreGlobals.h"
 #include "Clients/LiveLinkHubProvider.h"
+#include "LiveLinkHubLog.h"
 #include "LiveLinkHubSubjectSettings.h"
 #include "LiveLinkSettings.h"
+#include "Misc/App.h"
 #include "Misc/ConfigCacheIni.h"
 #include "Modules/ModuleManager.h"
 #include "PropertyEditorModule.h"
@@ -23,8 +26,10 @@ void FLiveLinkHubModule::StartupModule()
 {
 	TRACE_CPUPROFILER_EVENT_SCOPE(StartLiveLinkHub);
 
-	bCreateLiveLinkHubInstance =
-		GConfig->GetBoolOrDefault(TEXT("LiveLink"), TEXT("bCreateLiveLinkHubInstance"), false, GEngineIni);
+	const bool bConfigValue = bCreateLiveLinkHubInstance = GConfig->GetBoolOrDefault(TEXT("LiveLink"), TEXT("bCreateLiveLinkHubInstance"), false, GEngineIni);
+
+	bCreateLiveLinkHubInstance = bCreateLiveLinkHubInstance && !IsRunningCommandlet() && FApp::CanEverRender();
+	UE_LOG(LogLiveLinkHub, Display, TEXT("LiveLinkHubModule::StartupModule - LiveLinkHub instance %s (Config: %d)"), bCreateLiveLinkHubInstance ? TEXT("will be created.") : TEXT("will not be created."), (int32) bConfigValue);
 
 	if (bCreateLiveLinkHubInstance)
 	{
