Fix LiveLinkHub's CustomTimeStep rate divider slider sending too many updates when changing the slider value

diff --git a/Engine/Plugins/Animation/LiveLinkHub/Source/LiveLinkHub/Private/UI/Widgets/SLiveLinkTimecode.cpp b/Engine/Plugins/Animation/LiveLinkHub/Source/LiveLinkHub/Private/UI/Widgets/SLiveLinkTimecode.cpp
--- a/Engine/Plugins/Animation/LiveLinkHub/Source/LiveLinkHub/Private/UI/Widgets/SLiveLinkTimecode.cpp
+++ b/Engine/Plugins/Animation/LiveLinkHub/Source/LiveLinkHub/Private/UI/Widgets/SLiveLinkTimecode.cpp
@@ -396,11 +396,6 @@ void SLiveLinkTimecode::SetCustomTimeStepSubject(const FName SubjectName)
 
 void SLiveLinkTimecode::SetCustomTimeStepRateDivider(uint32 RateDivider)
 {
-	if (RateDivider == CachedFrameRateDivider)
-	{
-		return;
-	}
-
 	CachedFrameRateDivider = RateDivider;
 
 	const FLiveLinkHubCustomTimeStepSettings Settings = MakeCustomTimeStepSettings();
@@ -488,13 +483,17 @@ TSharedRef<SWidget> SLiveLinkTimecode::MakeFrameRateDivider()
 						.MinSliderValue(1)
 						.MaxSliderValue(100)
 						.OnValueChanged_Lambda([this](uint32 NewValue)
-							{
-								SetCustomTimeStepRateDivider(NewValue);
-							})
-						.Value_Lambda([this]()
-							{
-								return CachedFrameRateDivider;
-							})
+						{
+							SetCustomTimeStepRateDivider(NewValue);
+						})
+						.OnEndSliderMovement_Lambda([this](uint32 NewValue)
+						{
+							CachedFrameRateDivider = NewValue;
+						})
+					.Value_Lambda([this]()
+						{
+							return CachedFrameRateDivider;
+						})
 				]
 		];
 }
