Live Link Hub: Introduce ILiveLinkRecordingSessionInfo as slate metadata provider.

Values are serialized as part of LLH session JSON.

This is abstracted through a modular feature in order to support a hypothetical alternative implementation with e.g. Take Recorder as a backing store in Unreal Editor.

diff --git a/Engine/Plugins/Animation/LiveLinkHub/Source/LiveLinkHub/Private/Session/LiveLinkHubSessionData.h b/Engine/Plugins/Animation/LiveLinkHub/Source/LiveLinkHub/Private/Session/LiveLinkHubSessionData.h
--- a/Engine/Plugins/Animation/LiveLinkHub/Source/LiveLinkHub/Private/Session/LiveLinkHubSessionData.h
+++ b/Engine/Plugins/Animation/LiveLinkHub/Source/LiveLinkHub/Private/Session/LiveLinkHubSessionData.h
@@ -38,6 +38,18 @@ public:
 	UPROPERTY()
 	FLiveLinkHubCustomTimeStepSettings CustomTimeStepSettings;
 
+	/** Recording metadata - session name */
+	UPROPERTY()
+	FString RecordingSessionName;
+
+	/** Recording metadata - slate name */
+	UPROPERTY()
+	FString RecordingSlateName;
+
+	/** Recording metadata - take number */
+	UPROPERTY()
+	int32 RecordingTakeNumber;
+
 	/** Whether the hub should be used as a timecode source for connected clients. */
 	UPROPERTY()
 	bool bUseLiveLinkHubAsTimecodeSource = false;
diff --git a/Engine/Plugins/Animation/LiveLinkHub/Source/LiveLinkHub/Private/Session/LiveLinkHubSessionManager.h b/Engine/Plugins/Animation/LiveLinkHub/Source/LiveLinkHub/Private/Session/LiveLinkHubSessionManager.h
--- a/Engine/Plugins/Animation/LiveLinkHub/Source/LiveLinkHub/Private/Session/LiveLinkHubSessionManager.h
+++ b/Engine/Plugins/Animation/LiveLinkHub/Source/LiveLinkHub/Private/Session/LiveLinkHubSessionManager.h
@@ -10,6 +10,7 @@
 #include "Engine/Engine.h"
 #include "HAL/CriticalSection.h"
 #include "IDesktopPlatform.h"
+#include "ILiveLinkRecordingSessionInfo.h"
 #include "LiveLinkHubClient.h"
 #include "LiveLinkHubLog.h"
 #include "LiveLinkSourceSettings.h"
@@ -56,15 +57,21 @@ public:
 };
 
 class FLiveLinkHubSessionManager : public ILiveLinkHubSessionManager
+	, public ILiveLinkRecordingSessionInfo
 {
 public:
 	FLiveLinkHubSessionManager()
 	{
+		IModularFeatures::Get().RegisterModularFeature(ILiveLinkRecordingSessionInfo::GetModularFeatureName(), this);
+
 		FScopeLock Lock(&CurrentSessionCS);
 		CurrentSession = MakeShared<FLiveLinkHubSession>(OnClientAddedToSessionDelegate, OnClientRemovedFromSessionDelegate);
 	}
 
-	virtual ~FLiveLinkHubSessionManager() override = default;
+	virtual ~FLiveLinkHubSessionManager() override
+	{
+		IModularFeatures::Get().UnregisterModularFeature(ILiveLinkRecordingSessionInfo::GetModularFeatureName(), this);
+	}
 
 	//~ Begin LiveLinkHubSessionManager
 	virtual FOnClientAddedToSession& OnClientAddedToSession() override
@@ -217,10 +224,98 @@ public:
 	virtual const FString& GetLastConfigPath() const override
 	{
 		return LastConfigPath;
-	}
-	
+	}	
 	//~ End LiveLinkHubSessionManager
 
+	//~ Begin ILiveLinkRecordingSessionInfo interface
+	virtual FString GetSessionName() const override
+	{
+		FScopeLock Lock(&CurrentSessionCS);
+		if (ensure(CurrentSession))
+		{
+			return CurrentSession->SessionData->RecordingSessionName;
+		}
+
+		return FString();
+	}
+
+	virtual FString GetSlateName() const override
+	{
+		FScopeLock Lock(&CurrentSessionCS);
+		if (ensure(CurrentSession))
+		{
+			return CurrentSession->SessionData->RecordingSlateName;
+		}
+
+		return FString();
+	}
+
+	virtual int32 GetTakeNumber() const override
+	{
+		FScopeLock Lock(&CurrentSessionCS);
+		if (ensure(CurrentSession))
+		{
+			return CurrentSession->SessionData->RecordingTakeNumber;
+		}
+
+		return -1;
+	}
+
+	virtual bool SetSessionName(FStringView InSessionName) override
+	{
+		FScopeLock Lock(&CurrentSessionCS);
+		if (ensure(CurrentSession))
+		{
+			CurrentSession->SessionData->RecordingSessionName = InSessionName;
+			OnSessionNameChangedDelegate.Broadcast(InSessionName);
+			return true;
+		}
+
+		return false;
+	}
+
+	virtual bool SetSlateName(FStringView InSlateName) override
+	{
+		FScopeLock Lock(&CurrentSessionCS);
+		if (ensure(CurrentSession))
+		{
+			CurrentSession->SessionData->RecordingSlateName = InSlateName;
+			OnSlateNameChangedDelegate.Broadcast(InSlateName);
+			return true;
+		}
+
+		return false;
+	}
+
+	virtual bool SetTakeNumber(int32 InTakeNumber) override
+	{
+		FScopeLock Lock(&CurrentSessionCS);
+		if (ensure(CurrentSession))
+		{
+			CurrentSession->SessionData->RecordingTakeNumber = InTakeNumber;
+			OnTakeNumberChangedDelegate.Broadcast(InTakeNumber);
+			return true;
+		}
+
+		return false;
+	}
+
+	virtual FOnSessionStringChanged& OnSessionNameChanged() override
+	{
+		return OnSessionNameChangedDelegate;
+	}
+
+	virtual FOnSessionStringChanged& OnSlateNameChanged() override
+	{
+		return OnSlateNameChangedDelegate;
+	}
+
+	virtual FOnSessionIntChanged& OnTakeNumberChanged() override
+	{
+		return OnTakeNumberChangedDelegate;
+	}
+	//~ End ILiveLinkRecordingSessionInfo interface
+
 private:
 	/** Load a session from disk and restore its content. */
 	void RestoreSession(const FString& Path)
@@ -310,14 +405,28 @@ private:
 private:
 	/** Session that holds the current configuration of the hub (Clients, sources, subjects). */
 	TSharedPtr<FLiveLinkHubSession> CurrentSession;
+
 	/** Last path where we saved a session config file. */
 	FString LastConfigPath;
+
 	/** Delegate triggered when a client is added to the current session. */
 	FOnClientAddedToSession OnClientAddedToSessionDelegate;
+
 	/** Delegate triggered when a client is removed from the current session. */
 	FOnClientRemovedFromSession OnClientRemovedFromSessionDelegate;
+
 	/** Delegate triggered when the current session is changed. */
 	FOnActiveSessionChanged OnActiveSessionChangedDelegate;
+
+	/** Delegate triggered when the recording session name is changed. */
+	FOnSessionStringChanged OnSessionNameChangedDelegate;
+
+	/** Delegate triggered when the recording slate name is changed. */
+	FOnSessionStringChanged OnSlateNameChangedDelegate;
+
+	/** Delegate triggered when the recording take number is changed. */
+	FOnSessionIntChanged OnTakeNumberChangedDelegate;
+
 	/** Critical section used to synchronize access to the current session. */
 	mutable FCriticalSection CurrentSessionCS;
 };
diff --git a/Engine/Plugins/Animation/LiveLinkHub/Source/LiveLinkHub/Private/UI/Widgets/SLiveLinkHubTabViewBase.cpp b/Engine/Plugins/Animation/LiveLinkHub/Source/LiveLinkHub/Private/UI/Widgets/SLiveLinkHubTabViewBase.cpp
--- a/Engine/Plugins/Animation/LiveLinkHub/Source/LiveLinkHub/Private/UI/Widgets/SLiveLinkHubTabViewBase.cpp
+++ b/Engine/Plugins/Animation/LiveLinkHub/Source/LiveLinkHub/Private/UI/Widgets/SLiveLinkHubTabViewBase.cpp
@@ -11,6 +11,7 @@
 #include "Recording/LiveLinkHubRecordingController.h"
 #include "SLiveLinkHubSettings.h"
 #include "SLiveLinkHubStatusBar.h"
+#include "SLiveLinkRecordingSessionInfo.h"
 #include "SLiveLinkTimecode.h"
 
 #include "Widgets/Images/SImage.h"
@@ -18,6 +19,7 @@
 #include "Widgets/Input/SComboButton.h"
 #include "Widgets/Layout/SBorder.h"
 #include "Widgets/Layout/SSeparator.h"
+#include "Widgets/Layout/SSpacer.h"
 #include "Widgets/SBoxPanel.h"
 #include "Widgets/STimecode.h"
 #include "Widgets/Text/STextBlock.h"
@@ -91,16 +93,34 @@ void SLiveLinkHubTabViewBase::Construct(const FArguments& InArgs)
 					]
 #endif
 					+ SHorizontalBox::Slot()
-					.HAlign(HAlign_Right)
 					.VAlign(VAlign_Center)
+					.FillWidth(1.0)
 					[
 						SNew(SHorizontalBox)
 						+ SHorizontalBox::Slot()
+						.AutoWidth()
+						.VAlign(VAlign_Center)
+						[
+							SNew(SLiveLinkTimecode)
+						]
+						+ SHorizontalBox::Slot()
 						.Padding(2.f)
 						.VAlign(VAlign_Center)
 						.AutoWidth()
 						[
-							LiveLinkHub->GetRecordingController()->MakeRecordToolbarEntry()
+							SNew(SSeparator)
+							.Orientation(Orient_Vertical)
+						]
+						+ SHorizontalBox::Slot()
+						.FillWidth(1.0)
+						[
+							SNew(SSpacer)
+						]
+						+ SHorizontalBox::Slot()
+						.AutoWidth()
+						.VAlign(VAlign_Center)
+						[
+							SNew(SLiveLinkRecordingSessionInfo)
 						]
 						+ SHorizontalBox::Slot()
 						.Padding(2.f)
@@ -111,10 +131,16 @@ void SLiveLinkHubTabViewBase::Construct(const FArguments& InArgs)
 							.Orientation(Orient_Vertical)
 						]
 						+ SHorizontalBox::Slot()
+						.FillWidth(1.0)
+						[
+							SNew(SSpacer)
+						]
+						+ SHorizontalBox::Slot()
+						.Padding(2.f)
+						.VAlign(VAlign_Center)
 						.AutoWidth()
-	                    .VAlign(VAlign_Center)
 						[
-							SNew(SLiveLinkTimecode)
+							LiveLinkHub->GetRecordingController()->MakeRecordToolbarEntry()
 						]
 						+ SHorizontalBox::Slot()
 						.Padding(2.f)
diff --git a/Engine/Plugins/Animation/LiveLinkHub/Source/LiveLinkHub/Private/UI/Widgets/SLiveLinkRecordingSessionInfo.cpp b/Engine/Plugins/Animation/LiveLinkHub/Source/LiveLinkHub/Private/UI/Widgets/SLiveLinkRecordingSessionInfo.cpp
new file mode 100644
--- /dev/null
+++ b/Engine/Plugins/Animation/LiveLinkHub/Source/LiveLinkHub/Private/UI/Widgets/SLiveLinkRecordingSessionInfo.cpp
@@ -0,0 +1,110 @@
+// Copyright Epic Games, Inc. All Rights Reserved.
+
+#include "SLiveLinkRecordingSessionInfo.h"
+
+#include "ILiveLinkRecordingSessionInfo.h"
+#include "SPositiveActionButton.h"
+#include "Widgets/Input/SEditableTextBox.h"
+#include "Widgets/Text/STextBlock.h"
+#include "Widgets/SBoxPanel.h"
+
+
+#define LOCTEXT_NAMESPACE "SLiveLinkRecordingSessionInfo"
+
+
+void SLiveLinkRecordingSessionInfo::Construct(const FArguments& InArgs)
+{
+	const double IntraElementPadding = 6.0;
+
+	ILiveLinkRecordingSessionInfo& SessionInfo = ILiveLinkRecordingSessionInfo::Get();
+
+	ChildSlot
+	[
+		SNew(SHorizontalBox)
+		+ SHorizontalBox::Slot()
+		.VAlign(VAlign_Center)
+		.AutoWidth()
+		[
+			SNew(STextBlock)
+			.Margin(FMargin(IntraElementPadding, 0.0))
+			.Text(LOCTEXT("SessionEditLabel", "Session"))
+		]
+		+ SHorizontalBox::Slot()
+		.VAlign(VAlign_Center)
+		.AutoWidth()
+		[
+			SNew(SEditableTextBox)
+			.MinDesiredWidth(60)
+			.Text_Lambda([&SessionInfo] { return FText::FromString(SessionInfo.GetSessionName()); })
+			.OnTextCommitted_Lambda([&SessionInfo]
+				(const FText& InText, ETextCommit::Type)
+				{
+					SessionInfo.SetSessionName(InText.ToString());
+				})
+		]
+		+ SHorizontalBox::Slot()
+		.VAlign(VAlign_Center)
+		.AutoWidth()
+		[
+			SNew(STextBlock)
+			.Margin(FMargin(IntraElementPadding, 0.0))
+			.Text(LOCTEXT("SlateEditLabel", "Slate"))
+		]
+		+ SHorizontalBox::Slot()
+		.VAlign(VAlign_Center)
+		.AutoWidth()
+		[
+			SNew(SEditableTextBox)
+			.MinDesiredWidth(60)
+			.Text_Lambda([&SessionInfo] { return FText::FromString(SessionInfo.GetSlateName()); })
+			.OnTextCommitted_Lambda([&SessionInfo]
+				(const FText& InText, ETextCommit::Type)
+				{
+					SessionInfo.SetSlateName(InText.ToString());
+				})
+		]
+		+ SHorizontalBox::Slot()
+		.VAlign(VAlign_Center)
+		.AutoWidth()
+		[
+			SNew(STextBlock)
+			.Margin(FMargin(IntraElementPadding, 0.0))
+			.Text(LOCTEXT("TakeEditLabel", "Take"))
+		]
+		+ SHorizontalBox::Slot()
+		.VAlign(VAlign_Center)
+		.AutoWidth()
+		[
+			SNew(SEditableTextBox)
+			.MinDesiredWidth(20)
+			.Text_Lambda([&SessionInfo] { return FText::FromString(LexToString(SessionInfo.GetTakeNumber())); })
+			.OnTextCommitted_Lambda(
+				[&SessionInfo]
+				(const FText& InText, ETextCommit::Type)
+				{
+					int32 TakeNumber;
+					LexFromString(TakeNumber, *InText.ToString());
+					SessionInfo.SetTakeNumber(TakeNumber);
+				})
+		]
+		+ SHorizontalBox::Slot()
+		.VAlign(VAlign_Center)
+		.AutoWidth()
+		.Padding(IntraElementPadding, 0.0)
+		[
+			SNew(SPositiveActionButton)
+			.Icon(FAppStyle::Get().GetBrush("Icons.Plus"))
+			.ToolTipText(LOCTEXT("IncrementTake_ToolTip", "Increment the current take number"))
+			.OnClicked_Lambda(
+				[&SessionInfo]
+				() -> FReply
+				{
+					SessionInfo.SetTakeNumber(SessionInfo.GetTakeNumber() + 1);
+					return FReply::Handled();
+				})
+		]
+	];
+}
+
+
+#undef LOCTEXT_NAMESPACE
diff --git a/Engine/Plugins/Animation/LiveLinkHub/Source/LiveLinkHub/Private/UI/Widgets/SLiveLinkRecordingSessionInfo.h b/Engine/Plugins/Animation/LiveLinkHub/Source/LiveLinkHub/Private/UI/Widgets/SLiveLinkRecordingSessionInfo.h
new file mode 100644
--- /dev/null
+++ b/Engine/Plugins/Animation/LiveLinkHub/Source/LiveLinkHub/Private/UI/Widgets/SLiveLinkRecordingSessionInfo.h
@@ -0,0 +1,16 @@
+// Copyright Epic Games, Inc. All Rights Reserved.
+
+#pragma once
+
+#include "Widgets/SCompoundWidget.h"
+
+
+/** Display and edit the session name, slate name, take number. */
+class SLiveLinkRecordingSessionInfo : public SCompoundWidget
+{
+public:
+	SLATE_BEGIN_ARGS(SLiveLinkRecordingSessionInfo) {}
+	SLATE_END_ARGS()
+
+	void Construct(const FArguments& InArgs);
+};
diff --git a/Engine/Plugins/Animation/LiveLinkHub/Source/LiveLinkHub/Public/ILiveLinkRecordingSessionInfo.h b/Engine/Plugins/Animation/LiveLinkHub/Source/LiveLinkHub/Public/ILiveLinkRecordingSessionInfo.h
new file mode 100644
--- /dev/null
+++ b/Engine/Plugins/Animation/LiveLinkHub/Source/LiveLinkHub/Public/ILiveLinkRecordingSessionInfo.h
@@ -0,0 +1,43 @@
+﻿// Copyright Epic Games, Inc. All Rights Reserved.
+
+#pragma once
+
+#include "Containers/StringFwd.h"
+#include "Delegates/Delegate.h"
+#include "Features/IModularFeatures.h"
+
+
+/** Provides an interface for querying and setting recording session metadata. */
+class ILiveLinkRecordingSessionInfo : public IModularFeature
+{
+public:
+	DECLARE_MULTICAST_DELEGATE_OneParam(FOnSessionStringChanged, FStringView);
+	DECLARE_MULTICAST_DELEGATE_OneParam(FOnSessionIntChanged, int32);
+
+	/** Name of the modular feature. */
+	static FName GetModularFeatureName()
+	{
+		static const FName ModularFeatureName("LiveLinkRecordingSessionInfo");
+		return ModularFeatureName;
+	}
+
+	/** Convenience accessor for modular implementation. */
+	static ILiveLinkRecordingSessionInfo& Get()
+	{
+		IModularFeatures& ModularFeatures = IModularFeatures::Get();
+		ensure(ModularFeatures.GetModularFeatureImplementationCount(GetModularFeatureName()) == 1);
+		return ModularFeatures.GetModularFeature<ILiveLinkRecordingSessionInfo>(GetModularFeatureName());
+	}
+
+	virtual FString GetSessionName() const = 0;
+	virtual FString GetSlateName() const = 0;
+	virtual int32 GetTakeNumber() const = 0;
+
+	virtual bool SetSessionName(FStringView InSessionName) = 0;
+	virtual bool SetSlateName(FStringView InSlateName) = 0;
+	virtual bool SetTakeNumber(int32 InTakeNumber) = 0;
+
+	virtual FOnSessionStringChanged& OnSessionNameChanged() = 0;
+	virtual FOnSessionStringChanged& OnSlateNameChanged() = 0;
+	virtual FOnSessionIntChanged& OnTakeNumberChanged() = 0;
+};
