LiveLinkHub - Fix not being able to save and load frames with different sizes. Multiple static data recorded per track is now loaded properly. Note that changing static data per-frame is not officially supported. This data should remain constant for all frames.

diff --git a/Engine/Plugins/Animation/LiveLinkHub/Source/LiveLinkHub/Private/Recording/Implementations/LiveLinkFrameData.h b/Engine/Plugins/Animation/LiveLinkHub/Source/LiveLinkHub/Private/Recording/Implementations/LiveLinkFrameData.h
--- a/Engine/Plugins/Animation/LiveLinkHub/Source/LiveLinkHub/Private/Recording/Implementations/LiveLinkFrameData.h
+++ b/Engine/Plugins/Animation/LiveLinkHub/Source/LiveLinkHub/Private/Recording/Implementations/LiveLinkFrameData.h
@@ -101,10 +101,10 @@ namespace UE::LiveLinkHub::FrameData::Private
 		int64 RecordingStartFrameFilePosition = 0;
 		/** Maximum number of frames. */
 		int32 MaxFrames = 0;
-		/** Total size of the structure. */
-		int32 SerializedStructureSize = 0;
-		/** The size in bytes of each animation frame. */
-		int32 FrameDiskSize;
+		/** Frame offsets and sizes. [FrameOffsetBytes, FrameSizeBytes] */
+		TArray<TTuple<int32, int32>> FrameDiskSizes;
+		/** Whether the frame size is consistent throughout this animation. */
+		bool bHasConsistentFrameSize = false;
 		/** The last timestamp for this frame data. */
 		double LastTimestamp = 0.f;
 		/** The frame rate, based only off of number of frames and the last timestamp. */
@@ -113,7 +113,14 @@ namespace UE::LiveLinkHub::FrameData::Private
 		FFrameBufferCache BufferedCache;
 		/** Current iteration data while buffering. */
 		FFrameBufferIterationData BufferIterationData;
-	
+
+		/** The size in bytes of an animation frame. */
+		int32 GetFrameDiskSize(const int32 InFrameIdx) const
+		{
+			check(InFrameIdx >= 0 && InFrameIdx < FrameDiskSizes.Num());
+			return FrameDiskSizes[InFrameIdx].Value;
+		}
+		
 		/** Find the correct file offset based on the frame index. */
 		int64 GetFrameFilePosition(const int32 InFrameIdx) const
 		{
@@ -123,7 +130,8 @@ namespace UE::LiveLinkHub::FrameData::Private
 		/** Find the offset relative to local storage only, not accounting for disk position. */
 		int64 GetRelativeFrameFilePosition(const int32 InFrameIdx) const
 		{
-			return FrameDiskSize * InFrameIdx;
+			check(InFrameIdx >= 0 && InFrameIdx < FrameDiskSizes.Num());
+			return bHasConsistentFrameSize ? FrameDiskSizes[InFrameIdx].Value * InFrameIdx : FrameDiskSizes[InFrameIdx].Key;
 		};
 	};
 }
\ No newline at end of file
diff --git a/Engine/Plugins/Animation/LiveLinkHub/Source/LiveLinkHub/Private/Recording/Implementations/LiveLinkUAssetRecording.cpp b/Engine/Plugins/Animation/LiveLinkHub/Source/LiveLinkHub/Private/Recording/Implementations/LiveLinkUAssetRecording.cpp
--- a/Engine/Plugins/Animation/LiveLinkHub/Source/LiveLinkHub/Private/Recording/Implementations/LiveLinkUAssetRecording.cpp
+++ b/Engine/Plugins/Animation/LiveLinkHub/Source/LiveLinkHub/Private/Recording/Implementations/LiveLinkUAssetRecording.cpp
@@ -403,6 +403,16 @@ void ULiveLinkUAssetRecording::SaveFrameData(FArchive* InFileWriter, const FLive
 	// Write the struct name and size so it can be loaded later.
 	*InFileWriter << StructTypeName;
 
+	// Offset and size.
+	TArray<TTuple<int32, int32>> SerializedFrameSizes;
+	SerializedFrameSizes.AddDefaulted(NumFrames);
+
+	// Remember the position to write the frame size.
+	const uint64 SerializedFrameSizePosition = InFileWriter->Tell();
+	*InFileWriter << SerializedFrameSizes;
+	
+	SerializedFrameSizes.Reset();
+
 	// Write the frame header size.
 	{
 		const uint64 CurrentPosition = InFileWriter->Tell();
@@ -413,15 +423,15 @@ void ULiveLinkUAssetRecording::SaveFrameData(FArchive* InFileWriter, const FLive
 		InFileWriter->Seek(CurrentPosition);
 	}
 
-	// Remember the position to write the frame size.
-	const uint64 SerializedFrameSizePosition = InFileWriter->Tell();
-	int32 SerializedFrameSize = 0;
-	*InFileWriter << SerializedFrameSize;
-			
+	int32 RelativeStartPosition = 0;
+	
 	for (int32 FrameIdx = 0; FrameIdx < NumFrames; ++FrameIdx)
 	{
 		TSharedPtr<FInstancedStruct>& Frame = InBaseDataContainer.RecordedData[FrameIdx];
 		check(Frame.IsValid() && Frame->IsValid());
+
+		// Beginning of the frame data.
+		const int64 StartFramePosition = InFileWriter->Tell();
 			
 		// Write the frame index for streaming frames when loading.
 		*InFileWriter << FrameIdx;
@@ -431,25 +441,24 @@ void ULiveLinkUAssetRecording::SaveFrameData(FArchive* InFileWriter, const FLive
 		*InFileWriter << Timestamp;
 			
 		// Write the entire frame data.
-		uint64 SerializeDataStart = InFileWriter->Tell();
 		FObjectAndNameAsStringProxyArchive StructAr(*InFileWriter, false);
 		Frame->Serialize(StructAr);
 
 		// Store the serialized frame size, so we can write it once later.
 		{
-			int32 CurrentSerializedFrameSize = InFileWriter->Tell() - SerializeDataStart;
-			// Sanity check that the serialized frame size is consistent.
-			ensure(CurrentSerializedFrameSize == SerializedFrameSize || SerializedFrameSize == 0);
-			SerializedFrameSize = CurrentSerializedFrameSize;
+			const int64 EndFramePosition = InFileWriter->Tell();
+			const int32 CurrentSerializedFrameSize = EndFramePosition - StartFramePosition;
+			SerializedFrameSizes.Add({ RelativeStartPosition, CurrentSerializedFrameSize });
+			RelativeStartPosition += CurrentSerializedFrameSize;
 		}
 	}
 
-	if (SerializedFrameSize > 0)
+	if (SerializedFrameSizes.Num() > 0)
 	{
 		// Write the frame data offset at the beginning of the block.
 		const uint64 FinalOffset = InFileWriter->Tell();
 		InFileWriter->Seek(SerializedFrameSizePosition);
-		*InFileWriter << SerializedFrameSize;
+		*InFileWriter << SerializedFrameSizes;
 		InFileWriter->Seek(FinalOffset);
 	}
 }
@@ -499,7 +508,12 @@ void ULiveLinkUAssetRecording::LoadRecordingAsync(int32 InStartFrame, int32 InCu
 		AnimationData.ReadBulkDataPrimitive(LoadedRecordingVersion);
 
 		// If we modify the RecordingVersion we can perform import logic here.
-		ensure(LoadedRecordingVersion == RecordingVersion);
+		if (LoadedRecordingVersion != RecordingVersion)
+		{
+			UE_LOG(LogLiveLinkHub, Log, TEXT("Converting version %d to %d"), LoadedRecordingVersion, RecordingVersion);
+		}
+			
+		RecordingVersionBeingLoaded = LoadedRecordingVersion;
 			
 		// Process static data.
 		
@@ -517,7 +531,7 @@ void ULiveLinkUAssetRecording::LoadRecordingAsync(int32 InStartFrame, int32 InCu
 			}
 				
 			FLiveLinkRecordingStaticDataContainer& DataContainer = RecordingData.StaticData.FindChecked(*TemporaryFrameData.FrameDataSubjectKey.Get());
-			LoadFrameData(TemporaryFrameData, DataContainer, 0, 0, 1);
+			LoadFrameData(TemporaryFrameData, DataContainer, 0, 0, TemporaryFrameData.MaxFrames, /* bForceSequential */ true);
 
 			FScopeLock Lock(&DataContainerMutex);
 			MoveFrameDataToContainer(DataContainer, TemporaryFrameData);
@@ -537,7 +551,8 @@ void ULiveLinkUAssetRecording::LoadRecordingAsync(int32 InStartFrame, int32 InCu
 			}
 
 			// Offset to the end of this block if there is multiple NumFrameData.
-			AnimationData.SetBulkDataOffset(KeyPosition.GetFrameFilePosition(KeyPosition.MaxFrames));
+			const int32 EndBlockPosition = KeyPosition.GetFrameFilePosition(KeyPosition.MaxFrames - 1) + KeyPosition.GetFrameDiskSize(KeyPosition.MaxFrames - 1);
+			AnimationData.SetBulkDataOffset(EndBlockPosition);
 			FrameFileData.Add(*KeyPosition.FrameDataSubjectKey, MoveTemp(KeyPosition));
 		}
 	}
@@ -640,13 +655,24 @@ bool ULiveLinkUAssetRecording::LoadInitialFrameData(UE::LiveLinkHub::FrameData::
 		if (MaxFrames > 0)
 		{	
 			FString StructTypeName;
-			int32 SerializedStructureSize;
-
+			TArray<TTuple<int32, int32>> SerializedFrameSizes;
+			
 			Reader->GetMemoryReader() << StructTypeName;
 
-			AnimationData.ReadBulkDataPrimitive(SerializedStructureSize);
+			if (RecordingVersionBeingLoaded < UE::LiveLinkHub::Private::RecordingVersions::DynamicFrameSizes)
+			{
+				// Convert from 5.5.0 recordings where we expected all frames to be a constant size.
+				int32 SerializedFrameSize = 0;
+				AnimationData.ReadBulkDataPrimitive(SerializedFrameSize);
+				// Frame size consists of the frame index, timestamp, and frame struct data.
+				SerializedFrameSizes.Init({ 0, sizeof(int32) + sizeof(double) + SerializedFrameSize }, MaxFrames);
+			}
+			else
+			{
+				Reader->GetMemoryReader() << SerializedFrameSizes;
+			}
 
-			OutFrameData.SerializedStructureSize = SerializedStructureSize;
+			OutFrameData.FrameDiskSizes = MoveTemp(SerializedFrameSizes);
 			OutFrameData.RecordingStartFrameFilePosition = AnimationData.GetBulkDataOffset();
 
 			OutFrameData.LoadedStruct = FindObject<UScriptStruct>(nullptr, *StructTypeName, true);
@@ -656,12 +682,30 @@ bool ULiveLinkUAssetRecording::LoadInitialFrameData(UE::LiveLinkHub::FrameData::
 				return false;
 			}
 
-			// The size on disk for each frame-- consisting of the frame index, timestamp, and frame struct data.
-			OutFrameData.FrameDiskSize = (sizeof(int32) + sizeof(double) + SerializedStructureSize);
-
-			if (OutFrameData.FrameDiskSize > MaxFrameDiskSize)
+			// Determine max frame size and if there are different frame sizes.
+			if (OutFrameData.FrameDiskSizes.Num() > 0)
 			{
-				MaxFrameDiskSize = OutFrameData.FrameDiskSize;
+				bool bIsConsistentSize = true;
+				int32 MaxValue = TNumericLimits<int32>::Lowest();
+				int32 LastValue = OutFrameData.FrameDiskSizes[0].Value;
+				for (const TTuple<int32, int32>& Tuple : OutFrameData.FrameDiskSizes)
+				{
+					if (LastValue != Tuple.Value)
+					{
+						bIsConsistentSize = false;
+					}
+
+					LastValue = Tuple.Value;
+					
+					if (Tuple.Value > MaxValue)
+					{
+						MaxValue = Tuple.Value;
+					}
+				}
+
+				OutFrameData.bHasConsistentFrameSize = bIsConsistentSize;
+			
+				MaxFrameDiskSize = MaxValue;
 			}
 		}
 	}
@@ -686,7 +730,7 @@ bool ULiveLinkUAssetRecording::LoadInitialFrameData(UE::LiveLinkHub::FrameData::
 }
 
 void ULiveLinkUAssetRecording::LoadFrameData(UE::LiveLinkHub::FrameData::Private::FFrameMetaData& InFrameData, FLiveLinkRecordingBaseDataContainer& InDataContainer,
-                                             int32 RequestedStartFrame, int32 RequestedInitialFrame, int32 RequestedFramesToLoad)
+                                             int32 RequestedStartFrame, int32 RequestedInitialFrame, int32 RequestedFramesToLoad, bool bForceSequential)
 {
 	DECLARE_SCOPE_CYCLE_COUNTER(TEXT("ULiveLinkUAssetRecording::LoadFrameData"), STAT_ULiveLinkUAssetRecording_LoadFrameData, STATGROUP_LiveLinkHub);
 
@@ -813,12 +857,12 @@ void ULiveLinkUAssetRecording::LoadFrameData(UE::LiveLinkHub::FrameData::Private
 			RightFrameIdx++;
 		}
 		
-		auto AlternateLoadDirection = [&]()
+		auto AlternateLoadDirection = [&bLoadRight](const bool bRightOnly)
 		{
-			bLoadRight = !bLoadRight;
+			bLoadRight = bRightOnly ? true : !bLoadRight;
 		};
 
-		const int32 BufferBatchSize = GetDefault<ULiveLinkHubSettings>()->BufferBatchSize;
+		const int32 BufferBatchSize = bForceSequential ? MaxFrames : GetDefault<ULiveLinkHubSettings>()->BufferBatchSize;
 
 		struct FBulkDataBatch
 		{
@@ -866,8 +910,8 @@ void ULiveLinkUAssetRecording::LoadFrameData(UE::LiveLinkHub::FrameData::Private
 			}
 
 			// Make sure we're still within limits.
-			LeftBatchSize = FMath::Clamp(LeftBatchSize, 0, AvailableLeftFrames);
-			RightBatchSize = FMath::Clamp(RightBatchSize, 0, AvailableRightFrames);
+			LeftBatchSize = FMath::Clamp(LeftBatchSize, 0, InFrameData.bHasConsistentFrameSize ? AvailableLeftFrames : 1);
+			RightBatchSize = FMath::Clamp(RightBatchSize, 0, InFrameData.bHasConsistentFrameSize ? AvailableRightFrames : 1);
 			
 			// Determine the min/max frames for each batch.
 			if (LeftBatchSize > 0)
@@ -875,6 +919,11 @@ void ULiveLinkUAssetRecording::LoadFrameData(UE::LiveLinkHub::FrameData::Private
 				RawFramesLeftBatch.MinFrame = LeftFrameIdx - LeftBatchSize + 1;
 				RawFramesLeftBatch.MaxFrame = LeftFrameIdx;
 
+				if (RawFramesLeftBatch.MinFrame < 0)
+				{
+					RawFramesLeftBatch.MinFrame = 0;
+				}
+
 				// Shrink the batch to the first loaded frame, left to right since it's more likely frames for the left batch
 				// will already be loaded the more to the right (center).
 				for (int32 Idx = RawFramesLeftBatch.MinFrame; Idx < RawFramesLeftBatch.MaxFrame; ++Idx)
@@ -899,6 +948,11 @@ void ULiveLinkUAssetRecording::LoadFrameData(UE::LiveLinkHub::FrameData::Private
 				RawFramesRightBatch.MinFrame = RightFrameIdx;
 				RawFramesRightBatch.MaxFrame = RightFrameIdx + RightBatchSize - 1;
 
+				if (RawFramesRightBatch.MinFrame < 0)
+				{
+					RawFramesRightBatch.MinFrame = 0;
+				}
+				
 				// Shrink the batch to the first loaded frame, right to left since it's more likely frames for the right batch
 				// will already be loaded the more to the left (center).
 				for (int32 Idx = RawFramesRightBatch.MaxFrame; Idx >= RawFramesRightBatch.MinFrame; --Idx)
@@ -923,7 +977,7 @@ void ULiveLinkUAssetRecording::LoadFrameData(UE::LiveLinkHub::FrameData::Private
 		// Iterate through the entire range, loading frames from cache or disk. This runs until a batch cycle has completed
 		// or the entire range is loaded.
 		int32 FramesLoaded = 0;
-		while (RightFrameIdx < MaxFrames || LeftFrameIdx >= RequestedStartFrame)
+		while (RightFrameIdx < MaxFrames || (!bForceSequential && LeftFrameIdx >= RequestedStartFrame))
 		{
 			if (bCancelStream)
 			{
@@ -1044,7 +1098,7 @@ void ULiveLinkUAssetRecording::LoadFrameData(UE::LiveLinkHub::FrameData::Private
 			ensure(InDataContainer.Timestamps.Num() == InDataContainer.RecordedData.Num());
 			
 			FramesLoaded++;
-			AlternateLoadDirection();
+			AlternateLoadDirection(bForceSequential);
 			
 			// Break the loop when enough frames have finished constituting a batch of frames.
 			if (FramesLoaded > 0 && FramesLoaded % BufferBatchSize == 0)
@@ -1086,7 +1140,7 @@ bool ULiveLinkUAssetRecording::LoadFrameFromDisk(const int32 InFrame, const UE::
 	{
 		const int64 FramePosition = InFrameData.GetFrameFilePosition(InFrame);
 		AnimationData.SetBulkDataOffset(FramePosition);
-		Reader = AnimationData.CreateBulkDataMemoryReader(InFrameData.FrameDiskSize);
+		Reader = AnimationData.CreateBulkDataMemoryReader(InFrameData.GetFrameDiskSize(InFrame));
 	}
 
 	int32 ParsedFrameIdx = 0;
@@ -1128,7 +1182,7 @@ bool ULiveLinkUAssetRecording::LoadTimestampFromDisk(const int32 InFrame, const
 	const int64 FramePosition = InFrameData.GetFrameFilePosition(InFrame);
 	AnimationData.SetBulkDataOffset(FramePosition);
 			
-	const TSharedPtr<FLiveLinkHubBulkData::FScopedBulkDataMemoryReader> Reader = AnimationData.CreateBulkDataMemoryReader(InFrameData.FrameDiskSize);
+	const TSharedPtr<FLiveLinkHubBulkData::FScopedBulkDataMemoryReader> Reader = AnimationData.CreateBulkDataMemoryReader(InFrameData.GetFrameDiskSize(InFrame));
 
 	int32 ParsedFrameIdx = 0;
 	Reader->GetMemoryReader() << ParsedFrameIdx;
@@ -1150,6 +1204,9 @@ TSharedPtr<FLiveLinkHubBulkData::FScopedBulkDataMemoryReader> ULiveLinkUAssetRec
 {
 	DECLARE_SCOPE_CYCLE_COUNTER(TEXT("ULiveLinkUAssetRecording::LoadRawFramesFromDisk"), STAT_ULiveLinkUAssetRecording_LoadRawFramesFromDisk, STATGROUP_LiveLinkHub);
 
+	// Currently can only batch load multiple frames if frame size is consistent.
+	check(InFrameData.bHasConsistentFrameSize || InNumFrames <= 1);
+	
 	// Seek to the beginning of the frames to load.
 	const int64 FramePosition = InFrameData.GetFrameFilePosition(InFrame);
 	AnimationData.SetBulkDataOffset(FramePosition);
@@ -1159,7 +1216,7 @@ TSharedPtr<FLiveLinkHubBulkData::FScopedBulkDataMemoryReader> ULiveLinkUAssetRec
 	check(MaxFrames >= 1);
 
 	// Determine complete byte size to load.
-	const int32 MaxByteSize = MaxFrames * InFrameData.FrameDiskSize;
+	const int32 MaxByteSize = MaxFrames * InFrameData.GetFrameDiskSize(InFrame);
 	return AnimationData.CreateBulkDataMemoryReader(MaxByteSize);
 }
 
diff --git a/Engine/Plugins/Animation/LiveLinkHub/Source/LiveLinkHub/Private/Recording/Implementations/LiveLinkUAssetRecording.h b/Engine/Plugins/Animation/LiveLinkHub/Source/LiveLinkHub/Private/Recording/Implementations/LiveLinkUAssetRecording.h
--- a/Engine/Plugins/Animation/LiveLinkHub/Source/LiveLinkHub/Private/Recording/Implementations/LiveLinkUAssetRecording.h
+++ b/Engine/Plugins/Animation/LiveLinkHub/Source/LiveLinkHub/Private/Recording/Implementations/LiveLinkUAssetRecording.h
@@ -15,6 +15,13 @@
 class FMemoryReader;
 struct FLiveLinkPlaybackTracks;
 
+namespace UE::LiveLinkHub::Private::RecordingVersions
+{
+	constexpr int32 InitialVersion = 1;
+	/** Account for varying frame sizes and offsets. */
+	constexpr int32 DynamicFrameSizes = 2;
+}
+
 /**
  * Asset containing all animation data stored as bulk data. This is loaded async in chunks dependent on the playhead position.
  * Overall recording length, framerate, and frame indices are based on the maximum track length and farthest timestamp.
@@ -80,9 +87,19 @@ private:
 	/** Initial processing on a frame, finding the correct struct and offsets. The RecordingFileReader is assumed to be at the correct position. */
 	bool LoadInitialFrameData(UE::LiveLinkHub::FrameData::Private::FFrameMetaData& OutFrameData);
 	
-	/** Load frame data to a data container. */
+	/**
+	 * Load frame data to a data container. By default, it loads frames from the initial frame and then alternates in batches right then left.
+	 *
+	 * @param InFrameData The meta frame data being read and updated.
+	 * @param InDataContainer The data container for loaded frames to be stored.
+	 * @param RequestedStartFrame The farthest left start frame to load.
+	 * @param RequestedInitialFrame The first frame to be loaded, such as the midpoint, or the start frame.
+	 * @param RequestedFramesToLoad The total frames to try to load.
+	 * @param bForceSequential Whether to force a sequential load, left to right, rather than alternating.
+	 * In this case all frames from RequestedStartFrame to RequestedFramesToLoad are loaded.
+	 */
 	void LoadFrameData(UE::LiveLinkHub::FrameData::Private::FFrameMetaData& InFrameData, FLiveLinkRecordingBaseDataContainer& InDataContainer,
-		int32 RequestedStartFrame, int32 RequestedInitialFrame, int32 RequestedFramesToLoad);
+		int32 RequestedStartFrame, int32 RequestedInitialFrame, int32 RequestedFramesToLoad, bool bForceSequential = false);
 
 	/**
 	 * Attempt to load a frame from bulk data.
@@ -235,5 +252,8 @@ private:
 	bool bIsFullyLoaded = false;
 	
 	/** The current version of the recording. */
-	const int32 RecordingVersion = 1;
+	const int32 RecordingVersion = UE::LiveLinkHub::Private::RecordingVersions::DynamicFrameSizes;
+
+	/** The version being currently loaded. */
+	int32 RecordingVersionBeingLoaded = 0;
 };
