Convert LiveLinkHub to a WorkflowCentricApplication in order to use FApplicationModes for its layout management.

#rnx

diff --git a/Engine/Plugins/Animation/LiveLinkHub/Source/LiveLinkHub/LiveLinkHub.Build.cs b/Engine/Plugins/Animation/LiveLinkHub/Source/LiveLinkHub/LiveLinkHub.Build.cs
--- a/Engine/Plugins/Animation/LiveLinkHub/Source/LiveLinkHub/LiveLinkHub.Build.cs
+++ b/Engine/Plugins/Animation/LiveLinkHub/Source/LiveLinkHub/LiveLinkHub.Build.cs
@@ -40,6 +40,7 @@ public class LiveLinkHub : ModuleRules
 				"Slate",
 				"SlateCore",
 				"StandaloneRenderer",
+				"StatusBar",
 				"TimeManagement",
 				"ToolMenus",
 				"ToolWidgets",
diff --git a/Engine/Plugins/Animation/LiveLinkHub/Source/LiveLinkHub/Private/LiveLinkHub.cpp b/Engine/Plugins/Animation/LiveLinkHub/Source/LiveLinkHub/Private/LiveLinkHub.cpp
--- a/Engine/Plugins/Animation/LiveLinkHub/Source/LiveLinkHub/Private/LiveLinkHub.cpp
+++ b/Engine/Plugins/Animation/LiveLinkHub/Source/LiveLinkHub/Private/LiveLinkHub.cpp
@@ -13,6 +13,7 @@
 #include "LiveLinkEditorSettings.h"
 #include "LiveLinkHubClient.h"
 #include "LiveLinkHubCommands.h"
+#include "LiveLinkHubCreatorAppMode.h"
 #include "LiveLinkHubSubjectSettings.h"
 #include "LiveLinkHubTicker.h"
 #include "LiveLinkProviderImpl.h"
@@ -38,18 +39,31 @@ namespace LiveLinkHubInternalUtils
 	}
 }
 
+TSharedPtr<FLiveLinkHub> FLiveLinkHub::Get()
+{
+	TSharedPtr<FLiveLinkHub> Hub;
+	if (FLiveLinkHubModule* LiveLinkHubModule = FModuleManager::Get().GetModulePtr<FLiveLinkHubModule>("LiveLinkHub"))
+	{
+		Hub = LiveLinkHubModule->GetLiveLinkHub();
+	}
+
+	return Hub;
+}
+
 void FLiveLinkHub::Initialize(FLiveLinkHubTicker& Ticker)
 {
 	TRACE_CPUPROFILER_EVENT_SCOPE(FLiveLinkHub::Initialize);
+	
+	TSharedRef<FLiveLinkHub> LiveLinkHub = SharedThis<FLiveLinkHub>(this);
 
 	// We must register the livelink client first since we might rely on the modular feature to initialize the controllers/managers.
 	if (GetDefault<ULiveLinkHubSettings>()->bTickOnGameThread)
 	{
-		LiveLinkHubClient = MakeShared<FLiveLinkHubClient>(AsShared());
+		LiveLinkHubClient = MakeShared<FLiveLinkHubClient>(LiveLinkHub);
 	}
 	else
 	{
-		LiveLinkHubClient = MakeShared<FLiveLinkHubClient>(AsShared(), Ticker.OnTick());
+		LiveLinkHubClient = MakeShared<FLiveLinkHubClient>(LiveLinkHub, Ticker.OnTick());
 	}
 
 	IModularFeatures::Get().RegisterModularFeature(ILiveLinkClient::ModularFeatureName, LiveLinkHubClient.Get());
@@ -74,9 +88,8 @@ void FLiveLinkHub::Initialize(FLiveLinkHubTicker& Ticker)
 		TRACE_CPUPROFILER_EVENT_SCOPE(FLiveLinkHub::InitializeControllers);
 		RecordingController = MakeShared<FLiveLinkHubRecordingController>();
 		PlaybackController = MakeShared<FLiveLinkHubPlaybackController>();
-		RecordingListController = MakeShared<FLiveLinkHubRecordingListController>(AsShared());
+		RecordingListController = MakeShared<FLiveLinkHubRecordingListController>(LiveLinkHub);
 		ClientsController = MakeShared<FLiveLinkHubClientsController>(LiveLinkProvider.ToSharedRef());
-		CommandList = MakeShared<FUICommandList>();
 		SubjectController = MakeShared<FLiveLinkHubSubjectController>();
 	}
 
@@ -86,11 +99,23 @@ void FLiveLinkHub::Initialize(FLiveLinkHubTicker& Ticker)
 void FLiveLinkHub::InitializePostEngineInit()
 {
 	FLiveLinkHubCommands::Register();
-	BindCommands();
 
-	FString LiveLinkHubLayoutIni = GConfig->GetConfigFilename(TEXT("LiveLinkHubLayout"));
-	WindowController = MakeShared<FLiveLinkHubWindowController>(FLiveLinkHubWindowInitParams{ LiveLinkHubLayoutIni });
-	WindowController->RestoreLayout();
+	TabManager = FGlobalTabmanager::Get();
+
+	AddLiveLinkHubApplicationMode("CreatorMode", MakeShared<FLiveLinkHubCreatorAppMode>(SharedThis<FWorkflowCentricApplication>(this)));
+
+	for (ILiveLinkHubApplicationModeFactory* Factory : IModularFeatures::Get().GetModularFeatureImplementations<ILiveLinkHubApplicationModeFactory>(ILiveLinkHubApplicationModeFactory::ModularFeatureName))
+	{
+		TSharedRef<FLiveLinkHubApplicationMode> AppMode = Factory->CreateApplicationMode(SharedThis<FWorkflowCentricApplication>(this));
+		AddLiveLinkHubApplicationMode(AppMode->GetModeName(), MoveTemp(AppMode));
+	}
+
+	IModularFeatures::Get().OnModularFeatureRegistered().AddSP(this, &FLiveLinkHub::OnModularFeatureRegistered);
+
+	WindowController = MakeShared<FLiveLinkHubWindowController>();
+	
+	WindowController->RestoreLayout(AsShared());
+	SetCurrentMode("CreatorMode");
 
 	LiveLinkHubClient->OnStaticDataReceived_AnyThread().AddSP(this, &FLiveLinkHub::OnStaticDataReceived_AnyThread);
 	LiveLinkHubClient->OnFrameDataReceived_AnyThread().AddSP(this, &FLiveLinkHub::OnFrameDataReceived_AnyThread);
@@ -169,6 +194,20 @@ TSharedPtr<FLiveLinkHubPlaybackController> FLiveLinkHub::GetPlaybackController()
 	return PlaybackController;
 }
 
+void FLiveLinkHub::MapToolkitCommands()
+{
+	const FLiveLinkHubCommands& Commands = FLiveLinkHubCommands::Get();
+	TSharedPtr<FUICommandList> CommandList = GetToolkitCommands();
+
+	CommandList->MapAction(Commands.NewConfig, FExecuteAction::CreateSP(this, &FLiveLinkHub::NewConfig));
+	CommandList->MapAction(Commands.OpenConfig, FExecuteAction::CreateSP(this, &FLiveLinkHub::OpenConfig));
+	CommandList->MapAction(Commands.SaveConfigAs, FExecuteAction::CreateSP(this, &FLiveLinkHub::SaveConfigAs));
+	CommandList->MapAction(Commands.SaveConfig, FExecuteAction::CreateSP(this, &FLiveLinkHub::SaveConfig),
+		FCanExecuteAction::CreateSP(this, &FLiveLinkHub::CanSaveConfig));
+
+	FLiveLinkHubApplication::MapToolkitCommands();
+}
+
 void FLiveLinkHub::OnStaticDataReceived_AnyThread(const FLiveLinkSubjectKey& InSubjectKey, TSubclassOf<ULiveLinkRole> InRole, const FLiveLinkStaticDataStruct& InStaticDataStruct) const
 {
 	if (RecordingController->IsRecording())
@@ -197,16 +236,6 @@ void FLiveLinkHub::OnSubjectMarkedPendingKill_AnyThread(const FLiveLinkSubjectKe
 	LiveLinkProvider->RemoveSubject(OverridenName);
 }
 
-void FLiveLinkHub::BindCommands()
-{
-	const FLiveLinkHubCommands& Commands = FLiveLinkHubCommands::Get();
-	CommandList->MapAction(Commands.NewConfig, FExecuteAction::CreateSP(this, &FLiveLinkHub::NewConfig));
-	CommandList->MapAction(Commands.OpenConfig, FExecuteAction::CreateSP(this, &FLiveLinkHub::OpenConfig));
-	CommandList->MapAction(Commands.SaveConfigAs, FExecuteAction::CreateSP(this, &FLiveLinkHub::SaveConfigAs));
-	CommandList->MapAction(Commands.SaveConfig, FExecuteAction::CreateSP(this, &FLiveLinkHub::SaveConfig),
-		FCanExecuteAction::CreateSP(this, &FLiveLinkHub::CanSaveConfig));
-}
-
 void FLiveLinkHub::NewConfig()
 {
 	SessionManager->NewSession();
@@ -232,6 +261,17 @@ void FLiveLinkHub::OpenConfig()
 	SessionManager->RestoreSession();
 }
 
+void FLiveLinkHub::OnModularFeatureRegistered(const FName& Type, IModularFeature* ModularFeature)
+{
+	if (Type == ILiveLinkHubApplicationModeFactory::ModularFeatureName)
+	{
+		ILiveLinkHubApplicationModeFactory* Factory = static_cast<ILiveLinkHubApplicationModeFactory*>(ModularFeature);
+
+		TSharedRef<FLiveLinkHubApplicationMode> AppMode = Factory->CreateApplicationMode(SharedThis<FWorkflowCentricApplication>(this));
+		AddLiveLinkHubApplicationMode(AppMode->GetModeName(), MoveTemp(AppMode));
+	}
+}
+
 void FLiveLinkHub::RegisterLiveLinkHubSettings()
 {
 	if (ISettingsModule* SettingsModule = FModuleManager::GetModulePtr<ISettingsModule>("Settings"))
diff --git a/Engine/Plugins/Animation/LiveLinkHub/Source/LiveLinkHub/Private/LiveLinkHub.h b/Engine/Plugins/Animation/LiveLinkHub/Source/LiveLinkHub/Private/LiveLinkHub.h
--- a/Engine/Plugins/Animation/LiveLinkHub/Source/LiveLinkHub/Private/LiveLinkHub.h
+++ b/Engine/Plugins/Animation/LiveLinkHub/Source/LiveLinkHub/Private/LiveLinkHub.h
@@ -2,7 +2,10 @@
 
 #pragma once
 
+#include "LiveLinkHubApplication.h"
+
 #include "LiveLinkTypes.h"
+#include "LiveLinkHubModule.h"
 #include "Templates/SubclassOf.h"
 #include "Templates/UniquePtr.h"
 #include "Toolkits/FConsoleCommandExecutor.h"
@@ -38,9 +41,12 @@ public:
  * Implementation of the live link hub.
  * Contains the apps' different components and is responsible for handling communication between them.
  */
-class FLiveLinkHub : public ILiveLinkHub, public TSharedFromThis<FLiveLinkHub>
+class FLiveLinkHub : public ILiveLinkHub, public FLiveLinkHubApplication
 {
 public:
+	/** Get the livelinkhub instance that's held by LiveLinkHubModule. May return a nullptr when the module is shutting down. */ 
+	static TSharedPtr<FLiveLinkHub> Get();
+
 	virtual ~FLiveLinkHub() override;
 
 	//~ Begin ILiveLinkHub interface
@@ -48,6 +54,7 @@ public:
 	virtual bool IsRecording() const override;
 	//~ End ILiveLinkHub interface
 
+
 public:
 	/** 
 	 * Initialize LiveLinkHub. */
@@ -67,10 +74,13 @@ public:
 	/** Get the controller that manages clients. */
 	TSharedPtr<FLiveLinkHubClientsController> GetClientsController() const;
 	/** Get the live link hub command list. */
-	TSharedPtr<FUICommandList> GetCommandList() const { return CommandList; }
+	TSharedPtr<FUICommandList> GetCommandList() const { return GetToolkitCommands(); }
 	/** Get the session manager. */
 	TSharedPtr<ILiveLinkHubSessionManager> GetSessionManager() const;
-	
+
+	//~ FAssetEditorToolkit interface
+	virtual void MapToolkitCommands() override;
+
 private:
 	/** Complete the initialization after core engine systems have been initialized. */
 	void InitializePostEngineInit();
@@ -81,9 +91,6 @@ private:
 	void OnSubjectMarkedPendingKill_AnyThread(const FLiveLinkSubjectKey& InSubjectKey) const;
 	//~ LiveLink Client delegates
 
-	/** Bind all available live link hub commands. */
-	void BindCommands();
-
 	/** Register settings relevant to the livelink hub. */
 	void RegisterLiveLinkHubSettings();
 	/** Unregister settings relevant to the livelink hub. */
@@ -103,6 +110,9 @@ private:
 
 	/** Open an existing config. */
 	void OpenConfig();
+
+	/** Handles adding application modes that are registered after LLH has finished initialization. */
+	void OnModularFeatureRegistered(const FName& Type, class IModularFeature* ModularFeature);
 	
 private:
 	/** Implements the logic to manage the clients tabs. */
@@ -125,8 +135,6 @@ private:
 	TSharedPtr<FLiveLinkHubProvider> LiveLinkProvider;
 	/** Handles execution of commands. */
 	TUniquePtr<FConsoleCommandExecutor> CommandExecutor;
-	/** Available live link hub commands. */
-	TSharedPtr<FUICommandList> CommandList;
 	/** The last opened config path. */
 	FString LastConfigPath;
 
diff --git a/Engine/Plugins/Animation/LiveLinkHub/Source/LiveLinkHub/Private/LiveLinkHubApplication.cpp b/Engine/Plugins/Animation/LiveLinkHub/Source/LiveLinkHub/Private/LiveLinkHubApplication.cpp
new file mode 100644
--- /dev/null
+++ b/Engine/Plugins/Animation/LiveLinkHub/Source/LiveLinkHub/Private/LiveLinkHubApplication.cpp
@@ -0,0 +1,234 @@
+// Copyright Epic Games, Inc. All Rights Reserved.
+
+#include "LiveLinkHubApplication.h"
+
+#include "ISessionFrontendModule.h"
+#include "ISlateReflectorModule.h"
+#include "LiveLinkHubCreatorAppMode.h"
+#include "Misc/EngineBuildSettings.h"
+#include "StatusBarSubsystem.h"
+#include "Toolkits/AssetEditorToolkitMenuContext.h"
+#include "ToolMenus.h"
+#include "UI/Widgets/SApplicationModeSwitcher.h"
+#include "WorkflowOrientedApp/ApplicationMode.h"
+#include "WorkflowOrientedApp/WorkflowTabManager.h"
+
+#define LOCTEXT_NAMESPACE "LiveLinkHubApplication"
+
+FLiveLinkHubApplication::FLiveLinkHubApplication()
+	: FWorkflowCentricApplication()
+{
+	LayoutIni = TEXT("LiveLinkHubLayout");
+
+	// We set an extender here and collect toolbar widgets 
+	TSharedPtr<FExtender> Extender = MakeShared<FExtender>();
+	Extender->AddToolBarExtension(
+		"Asset",
+		EExtensionHook::Before,
+		GetToolkitCommands(),
+		FToolBarExtensionDelegate::CreateRaw(this, &FLiveLinkHubApplication::AddToolbarExtenders)
+	);
+
+	AddToolbarExtender(Extender);
+}
+
+void FLiveLinkHubApplication::AddToolbarExtenders(FToolBarBuilder&)
+{
+	if (TSharedPtr<FApplicationMode> Mode = GetCurrentModePtr())
+	{
+		TArray<TSharedRef<SWidget>> Widgets = StaticCastSharedPtr<FLiveLinkHubApplicationMode>(GetCurrentModePtr())->GetToolbarWidgets();
+
+		for (const TSharedRef<SWidget>& Widget : Widgets)
+		{
+			AddToolbarWidget(Widget);
+		}
+	}
+}
+
+void FLiveLinkHubApplication::AddLiveLinkHubApplicationMode(FName ModeName, TSharedRef<FLiveLinkHubApplicationMode> Mode)
+{
+	CachedModeInfo.Add(ModeName, FLiveLinkHubAppModeInfo{ Mode->GetModeIcon(), Mode->GetDisplayName()});
+	FWorkflowCentricApplication::AddApplicationMode(ModeName, MoveTemp(Mode));
+}
+
+void FLiveLinkHubApplication::SetCurrentMode(FName NewMode)
+{
+	FWorkflowCentricApplication::SetCurrentMode(NewMode);
+	AppModeChangedDelegate.Broadcast(NewMode);
+}
+
+TOptional<FLiveLinkHubAppModeInfo> FLiveLinkHubApplication::GetModeInfo(FName Mode) const
+{
+	TOptional<FLiveLinkHubAppModeInfo> ModeInfo;
+
+	if (const FLiveLinkHubAppModeInfo* InfoPtr = CachedModeInfo.Find(Mode))
+	{
+		ModeInfo = *InfoPtr;
+	}
+
+	return ModeInfo;
+}
+
+TArray<FName> FLiveLinkHubApplication::GetApplicationModes() const
+{
+	TArray<FName> ApplicationModeNames;
+	CachedModeInfo.GenerateKeyArray(ApplicationModeNames);
+	return ApplicationModeNames;
+}
+
+UToolMenu* FLiveLinkHubApplication::GenerateCommonActionsToolbar(FToolMenuContext& MenuContext)
+{
+	UToolMenus* ToolMenus = UToolMenus::Get();
+	FName ToolBarName = "AssetEditorToolbar.CommonActions";
+
+	UToolMenu* FoundMenu = ToolMenus->FindMenu(ToolBarName);
+
+	if (!FoundMenu || !FoundMenu->IsRegistered())
+	{
+		FoundMenu = ToolMenus->RegisterMenu(ToolBarName, NAME_None, EMultiBoxType::SlimHorizontalToolBar);
+		FoundMenu->StyleName = "AssetEditorToolbar";
+
+		FToolMenuSection& Section = FoundMenu->AddSection("CommonActions");
+
+		Section.AddEntry(FToolMenuEntry::InitWidget("AppModeSwitcher", SNew(SApplicationModeSwitcher), LOCTEXT("AppModeSwitcherLabel", "ApplicationMode Switcher")));
+		Section.AddSeparator(NAME_None);
+	}
+
+	return ToolMenus->GenerateMenu(ToolBarName, MenuContext);
+}
+
+void FLiveLinkHubApplication::RegisterFileMenu(UToolMenu* Menu)
+{
+	UToolMenu* FileMenu = Menu->AddSubMenu(
+			"MainMenu",
+			NAME_None,
+			FLiveLinkHubApplicationMode::FileMenuExtensionPoint, // Prevents showing other submenus that would've been registered to File, which we don't care about.
+			LOCTEXT("FileMenu", "File"),
+			FText::GetEmpty()
+		);
+	
+	FToolMenuSection& FileAssetSection = FileMenu->AddSection("FileOpen", LOCTEXT("FileOpenHeading", "Open"), FToolMenuInsert(NAME_None, EToolMenuInsertType::First));
+
+	{
+		FileAssetSection.AddMenuEntry(FLiveLinkHubCommands::Get().NewConfig);
+		FileAssetSection.AddMenuEntry(FLiveLinkHubCommands::Get().OpenConfig);
+	}
+
+	FToolMenuSection& OpenAssetSection = FileMenu->AddSection("FileSave", LOCTEXT("FileSaveHeading", "Save"), FToolMenuInsert("FileOpen", EToolMenuInsertType::After));
+	{
+		OpenAssetSection.AddMenuEntry(FLiveLinkHubCommands::Get().SaveConfig);
+		OpenAssetSection.AddMenuEntry(FLiveLinkHubCommands::Get().SaveConfigAs);
+	}
+}
+
+void FLiveLinkHubApplication::RegisterMainMenu()
+{
+	static const FName MainMenuName("MainFrame.MainMenu");
+	UToolMenus* ToolMenus = UToolMenus::Get();
+
+	if (ToolMenus->IsMenuRegistered(MainMenuName))
+	{
+		return;
+	}
+
+	UToolMenu* MenuBar = ToolMenus->RegisterMenu(MainMenuName, NAME_None, EMultiBoxType::MenuBar);
+	MenuBar->StyleName = FName("WindowMenuBar");
+
+	RegisterFileMenu(MenuBar);
+
+	const bool bShowDevTools = FParse::Param(FCommandLine::Get(), TEXT("Development"))
+		|| FEngineBuildSettings::IsInternalBuild();
+
+	if (bShowDevTools)
+	{
+		UToolMenu* ToolsEntry = MenuBar->AddSubMenu(
+			"MainMenu",
+			NAME_None,
+			"Tools",
+			LOCTEXT("ToolsMenu", "Tools"),
+			FText::GetEmpty()
+		);
+
+		ToolsEntry->AddDynamicSection(NAME_None, FNewSectionConstructChoice{FNewToolMenuDelegateLegacy::CreateRaw(this, &FLiveLinkHubApplication::CreateToolsMenu)});
+	}
+
+	UToolMenu* WindowEntry = MenuBar->AddSubMenu(
+		"MainMenu",
+		NAME_None,
+		"Window",
+		LOCTEXT("WindowMenu", "Window"),
+		FText::GetEmpty()
+	);
+
+	WindowEntry->AddDynamicSection(NAME_None, FNewSectionConstructChoice{FNewToolMenuDelegateLegacy::CreateRaw(this, &FLiveLinkHubApplication::FillWindowMenu)});
+}
+
+void FLiveLinkHubApplication::CreateToolsMenu(FMenuBuilder& MenuBuilder, UToolMenu*) const
+{
+	MenuBuilder.BeginSection("Development", LOCTEXT("DevelopmentHeader", "Development"));
+
+#if WITH_UNREAL_TARGET_DEVELOPER_TOOLS
+	MenuBuilder.AddMenuEntry(
+		LOCTEXT("FileMenu_Development_AutomationTools", "Automation Tools"),
+		FText(),
+		FSlateIcon(),
+		FExecuteAction::CreateStatic([]()
+			{
+				ISessionFrontendModule& SessionFrontend =
+					FModuleManager::LoadModuleChecked<ISessionFrontendModule>("SessionFrontend");
+				SessionFrontend.InvokeSessionFrontend(FName("AutomationPanel"));
+			})
+	);
+#endif // #if WITH_UNREAL_TARGET_DEVELOPER_TOOLS
+	MenuBuilder.AddMenuEntry(
+		LOCTEXT("FileMenu_Development_WidgetReflector", "Widget Reflector"),
+		FText(),
+		FSlateIcon(),
+		FExecuteAction::CreateStatic([]()
+			{
+				ISlateReflectorModule& SlateReflector =
+					FModuleManager::LoadModuleChecked<ISlateReflectorModule>("SlateReflector");
+				SlateReflector.DisplayWidgetReflector();
+			})
+	);
+	MenuBuilder.EndSection();
+}
+
+void FLiveLinkHubApplication::FillWindowMenu(FMenuBuilder& MenuBuilder, UToolMenu*)
+{
+	GetTabManager()->PopulateLocalTabSpawnerMenu(MenuBuilder);
+}
+
+UE::Editor::Toolbars::ECreateStatusBarOptions FLiveLinkHubApplication::GetStatusBarCreationOptions() const
+{
+	return UE::Editor::Toolbars::ECreateStatusBarOptions::HideContentBrowser | UE::Editor::Toolbars::ECreateStatusBarOptions::HideSourceControl;
+}
+
+TSharedPtr<SWidget> FLiveLinkHubApplication::CreateMenuBar(const TSharedPtr<FTabManager>& InTabManager, const FName InMenuName, FToolMenuContext& InToolMenuContext)
+{
+	RegisterMainMenu();
+
+	InToolMenuContext.AppendCommandList(GetToolkitCommands());
+
+	USlateTabManagerContext* ContextObject = NewObject<USlateTabManagerContext>();
+	ContextObject->TabManager = TabManager;
+	InToolMenuContext.AddObject(ContextObject);
+
+	// Create the menu bar
+	TSharedRef<SWidget> MenuBarWidget = UToolMenus::Get()->GenerateWidget(InMenuName, InToolMenuContext);
+	if (MenuBarWidget != SNullWidget::NullWidget)
+	{
+		// Tell tab-manager about the multi-box for platforms with a global menu bar
+		TSharedRef<SMultiBoxWidget> MultiBoxWidget = StaticCastSharedRef<SMultiBoxWidget>(MenuBarWidget);
+		TabManager->SetMenuMultiBox(ConstCastSharedRef<FMultiBox>(MultiBoxWidget->GetMultiBox()), MultiBoxWidget);
+	}
+
+	return SNullWidget::NullWidget;
+}
+
+FText FLiveLinkHubApplication::GetToolkitName() const
+{
+	return LOCTEXT("LiveLinkHubLabel", "Live Link Hub");
+}
+
+#undef LOCTEXT_NAMESPACE /*LiveLinkHubApplication*/
diff --git a/Engine/Plugins/Animation/LiveLinkHub/Source/LiveLinkHub/Private/LiveLinkHubApplication.h b/Engine/Plugins/Animation/LiveLinkHub/Source/LiveLinkHub/Private/LiveLinkHubApplication.h
new file mode 100644
--- /dev/null
+++ b/Engine/Plugins/Animation/LiveLinkHub/Source/LiveLinkHub/Private/LiveLinkHubApplication.h
@@ -0,0 +1,82 @@
+// Copyright Epic Games, Inc. All Rights Reserved.
+
+#pragma once
+
+#include "WorkflowOrientedApp/WorkflowCentricApplication.h"
+#include "Math/Color.h"
+#include "Misc/Optional.h"
+
+class FLiveLinkHubApplicationMode;
+
+/** Cached information about a given application mode. Used to populate the mode switcher widget. */
+struct FLiveLinkHubAppModeInfo
+{
+	FSlateIcon Icon;
+	FText DisplayName;
+};
+
+DECLARE_MULTICAST_DELEGATE_OneParam(FOnApplicationModeChanged, /*CurrentMode*/FName);
+
+class FLiveLinkHubApplication : public FWorkflowCentricApplication
+{
+public:
+	FLiveLinkHubApplication();
+	virtual ~FLiveLinkHubApplication() {}
+
+	/** Get the delegate called when the app mode has changed. */
+	FOnApplicationModeChanged& OnApplicationModeChanged() { return AppModeChangedDelegate; }
+
+	/** Get information about a given mode. Returns an empty optional if the mode wasn't found or if it wasn't added to the app using AddLiveLinkHubApplicationMode. */
+	TOptional<FLiveLinkHubAppModeInfo> GetModeInfo(FName Mode) const;
+	
+	/** Get the list of modes that were registered with the app. */
+	TArray<FName> GetApplicationModes() const;
+
+	/** Add a LiveLinkHub application mode. */
+	void AddLiveLinkHubApplicationMode(FName ModeName, TSharedRef<FLiveLinkHubApplicationMode> Mode);
+
+	//~ Begin FWorkflowCentricApplication interface
+	virtual void SetCurrentMode(FName NewMode) override;
+	//~ End FWorkflowCentricApplication interface
+
+protected:
+	//~ Begin FWorkflowCentricApplication interface
+	virtual FName GetToolkitFName() const override { return "LiveLinkHub"; }
+	virtual FText GetBaseToolkitName() const override { return INVTEXT("LiveLinkHub Editor"); }
+	virtual FString GetWorldCentricTabPrefix() const override { return TEXT("LiveLinkHub"); }
+	FLinearColor GetWorldCentricTabColorScale() const
+	{
+		return FLinearColor::White;
+	}
+	//~ End FWorkflowCentricApplication interface
+
+	//~ Begin AssetEditorToolkit interface
+	virtual UToolMenu* GenerateCommonActionsToolbar(FToolMenuContext& MenuContext) override;
+	virtual UE::Editor::Toolbars::ECreateStatusBarOptions GetStatusBarCreationOptions() const override;
+	virtual TSharedPtr<SWidget> CreateMenuBar(const TSharedPtr<FTabManager>& InTabManager, const FName InMenuName, FToolMenuContext& InToolMenuContext) override;
+	virtual void MapToolkitCommands() {}
+	virtual FText GetToolkitName() const override;
+	//~ End AssetEditorToolkit interface
+
+private:
+	/** Register and populate the file menu with the config commands. */
+	void RegisterFileMenu(UToolMenu* Menu);
+	/** Create the main menu bar. */
+	void RegisterMainMenu();
+	/** When in dev mode, creates the tool menu that contains the widget reflector and automation window. */
+	void CreateToolsMenu(FMenuBuilder& MenuBuilder, UToolMenu*) const;
+	void FillWindowMenu(FMenuBuilder& MenuBuilder, UToolMenu*);
+
+	/** Used to populate the status toolbar with widgets of the current mode. */
+	void AddToolbarExtenders(FToolBarBuilder&);
+
+private:
+	/** Don't allow LiveLinkHub to call this directly since it wouldn't register the cached mode info. */
+	using FWorkflowCentricApplication::AddApplicationMode;
+
+	/** Delegate called when the active app mode has changed. */
+	FOnApplicationModeChanged AppModeChangedDelegate;
+
+	/** Map of mode name to cached mode info. */
+	TMap<FName, FLiveLinkHubAppModeInfo> CachedModeInfo;
+};
\ No newline at end of file
diff --git a/Engine/Plugins/Animation/LiveLinkHub/Source/LiveLinkHub/Private/LiveLinkHubApplicationMode.cpp b/Engine/Plugins/Animation/LiveLinkHub/Source/LiveLinkHub/Private/LiveLinkHubApplicationMode.cpp
new file mode 100644
--- /dev/null
+++ b/Engine/Plugins/Animation/LiveLinkHub/Source/LiveLinkHub/Private/LiveLinkHubApplicationMode.cpp
@@ -0,0 +1,52 @@
+// Copyright Epic Games, Inc. All Rights Reserved.
+
+#include "LiveLinkHubApplicationMode.h"
+
+#include "UI/Widgets/SLiveLinkHubSettings.h"
+#include "WorkflowOrientedApp/WorkflowCentricApplication.h"
+#include "ToolMenus.h"
+
+const FName ILiveLinkHubApplicationModeFactory::ModularFeatureName = "LiveLinkHubApplicationModeFactory";
+const FName FLiveLinkHubApplicationMode::FileMenuExtensionPoint = "LiveLinkHubFileMenu";
+
+FLiveLinkHubApplicationMode::FLiveLinkHubApplicationMode(FName InApplicationMode, FText InDisplayName, TSharedPtr<FWorkflowCentricApplication> InApp)
+	: FApplicationMode(InApplicationMode)
+{
+	LayoutIni = TEXT("LiveLinkHubLayout");
+	WeakApp = InApp;
+	DisplayName = InDisplayName;
+
+	WorkspaceMenuCategory = FWorkspaceItem::NewGroup(InDisplayName);
+}
+
+TArray<TSharedRef<SWidget>> FLiveLinkHubApplicationMode::GetToolbarWidgets()
+{
+	TArray<TSharedRef<SWidget>> Widgets = GetToolbarWidgets_Impl();
+
+	// LiveLinkHub Settings should always be last.
+	Widgets.Add(SNew(SLiveLinkHubSettings));
+
+	return Widgets;
+}
+
+void FLiveLinkHubApplicationMode::PostActivateMode()
+{
+	FToolMenuOwnerScoped OwnerScoped(GetModeName());
+
+	TArray<TSharedRef<SWidget>> Widgets = GetStatusBarWidgets_Impl();
+
+	UToolMenu* Menu = UToolMenus::Get()->ExtendMenu(TEXT("LiveLinkHub.ToolBar"));
+	FToolMenuSection& SessionConfigSection = Menu->AddSection(TEXT("ModeWidgets"), FText::GetEmpty(), FToolMenuInsert(NAME_None, EToolMenuInsertType::First));
+
+	for (const TSharedRef<SWidget>& Widget : Widgets)
+	{
+		SessionConfigSection.AddEntry(
+			FToolMenuEntry::InitWidget(TEXT("ModeWidget"), Widget, FText::GetEmpty(), true, false)
+		);
+	}
+}
+
+void FLiveLinkHubApplicationMode::PreDeactivateMode()
+{
+	UToolMenus::Get()->UnregisterOwner(GetModeName());
+}
\ No newline at end of file
diff --git a/Engine/Plugins/Animation/LiveLinkHub/Source/LiveLinkHub/Private/LiveLinkHubCreatorAppMode.h b/Engine/Plugins/Animation/LiveLinkHub/Source/LiveLinkHub/Private/LiveLinkHubCreatorAppMode.h
new file mode 100644
--- /dev/null
+++ b/Engine/Plugins/Animation/LiveLinkHub/Source/LiveLinkHub/Private/LiveLinkHubCreatorAppMode.h
@@ -0,0 +1,490 @@
+// Copyright Epic Games, Inc. All Rights Reserved.
+
+#pragma once
+
+#include "Clients/LiveLinkHubClientsController.h"
+#include "Features/IModularFeatures.h"
+#include "Framework/Docking/TabManager.h"
+#include "Framework/MultiBox/MultiBoxBuilder.h"
+#include "IDetailsView.h"
+#include "ILiveLinkClient.h"
+#include "LiveLinkClientPanelViews.h"
+#include "LiveLinkClientPanelToolbar.h"
+#include "LiveLinkHub.h"
+#include "LiveLinkHubApplicationMode.h"
+#include "LiveLinkHubCommands.h"
+#include "LiveLinkHubModule.h"
+#include "LiveLinkPanelController.h"
+#include "LiveLinkTypes.h"
+#include "Modules/ModuleManager.h"
+#include "Recording/LiveLinkHubRecordingListController.h"
+#include "Subjects/LiveLinkHubSubjectController.h"
+#include "ToolMenus.h"
+#include "ToolMenuSection.h"
+#include "UI/Widgets/SLiveLinkHubSettings.h"
+#include "UI/Widgets/SLiveLinkRecordingSessionInfo.h"
+#include "UI/Widgets/SLiveLinkTimecode.h"
+#include "Widgets/Docking/SDockTab.h"
+#include "Widgets/Input/SCheckBox.h"
+#include "Widgets/Input/SComboButton.h"
+#include "Widgets/Layout/SSplitter.h"
+#include "Widgets/Text/STextBlock.h"
+#include "WorkflowOrientedApp/ApplicationMode.h"
+#include "WorkflowOrientedApp/WorkflowTabFactory.h"
+#include "WorkflowOrientedApp/WorkflowTabManager.h"
+
+
+#define LOCTEXT_NAMESPACE "LiveLinkHubCreatorMode"
+
+namespace UE::LiveLinkHub
+{
+	static const FName LiveLinkStyleName = "LiveLinkStyle";
+
+	static const FName SourcesTabId = "SourcesTabId";
+	static const FName SourceDetailsTabId = "SourceDetailsTabId";
+	static const FName SubjectsTabId = "SubjectsTabId";
+	static const FName PlaybackTabId = "PlaybackTabId";
+	static const FName ClientsTabId = "ClientsTabId";
+	static const FName ClientDetailsTabId = "ClientDetailsTabId";
+}
+
+struct FLiveLinkHubSourcesTabSummoner : public FWorkflowTabFactory
+{
+public:
+	FLiveLinkHubSourcesTabSummoner(TSharedPtr<class FAssetEditorToolkit> InHostingApp, const TSharedPtr<FLiveLinkPanelController>& InPanelController)
+		: FWorkflowTabFactory(UE::LiveLinkHub::SourcesTabId, MoveTemp(InHostingApp))
+		, WeakPanelController(InPanelController)
+	{
+		TabLabel = LOCTEXT("SourcesTabLabel", "Sources");
+		TabIcon = FSlateIcon(UE::LiveLinkHub::LiveLinkStyleName, TEXT("LiveLinkHub.Subjects.Icon"));
+
+		bIsSingleton = true;
+	}
+
+	virtual TSharedRef<SWidget> CreateTabBody(const FWorkflowTabSpawnInfo& Info) const override
+	{
+		FLiveLinkClient* Client = (FLiveLinkClient*)&IModularFeatures::Get().GetModularFeature<ILiveLinkClient>(ILiveLinkClient::ModularFeatureName);
+		FLiveLinkHubModule& LiveLinkHubModule = FModuleManager::Get().GetModuleChecked<FLiveLinkHubModule>("LiveLinkHub");
+
+		TSharedRef<SWidget> SourcesListView = SNullWidget::NullWidget;
+		if (TSharedPtr<FLiveLinkPanelController> PanelController = WeakPanelController.Pin())
+		{
+			SourcesListView = PanelController->SourcesView->SourcesListView.ToSharedRef();
+		}
+
+		TSharedRef<SWidget> CustomToolbarHeader = SNew(SHorizontalBox)
+			+ SHorizontalBox::Slot()
+			.AutoWidth()
+			.Padding(0.f, 5.f, 0.f, 5.f)
+			[
+				SNew(SImage)
+					.ColorAndOpacity(FSlateColor::UseForeground())
+					.Image(FSlateIcon("LiveLinkStyle", "LiveLinkHub.Sources.Icon").GetIcon())
+			]
+			+ SHorizontalBox::Slot()
+			.HAlign(HAlign_Left)
+			.Padding(FMargin(4.0, 2.0))
+			[
+				SNew(STextBlock)
+					.Font(DEFAULT_FONT("Regular", 14))
+					.Text(LOCTEXT("SourcesHeaderText", "Sources"))
+			];
+
+
+		return SNew(SVerticalBox)
+			+ SVerticalBox::Slot()
+			.Padding(FMargin(4.0f, 0.f, 0.f, 0.f))
+			.AutoHeight()
+			[
+				SNew(SLiveLinkClientPanelToolbar, Client)
+					.SourceButtonAlignment(HAlign_Right)
+					.ParentWindow(LiveLinkHubModule.GetLiveLinkHub()->GetRootWindow())
+					.ShowPresetPicker(false)
+					.ShowSettings(false)
+					.CustomHeader(CustomToolbarHeader)
+					.IsEnabled_Lambda([this]()
+						{
+							const FLiveLinkHubModule& LiveLinkHubModule = FModuleManager::Get().GetModuleChecked<FLiveLinkHubModule>("LiveLinkHub");
+							return !LiveLinkHubModule.GetPlaybackController()->IsInPlayback() && !LiveLinkHubModule.GetRecordingController()->IsRecording();
+						})
+			]
+			+ SVerticalBox::Slot()
+			.FillHeight(1.f)
+			.Padding(FMargin(0.0f, 4.0f, 0.0f, 0.0f))
+			[
+				SourcesListView
+			];
+	}
+
+	virtual FText GetTabToolTipText(const FWorkflowTabSpawnInfo& Info) const override
+	{
+		return LOCTEXT("SourcesTabToolTip", "Displays the list of Live Link sources.");
+	}
+
+private:
+	/** Panel controller used to fetch the list of sources. */
+	TWeakPtr<FLiveLinkPanelController> WeakPanelController;
+};
+
+
+struct FLiveLinkHubSourceDetailsTabSummoner : public FWorkflowTabFactory
+{
+public:
+	FLiveLinkHubSourceDetailsTabSummoner(TSharedPtr<class FAssetEditorToolkit> InHostingApp, const TSharedPtr<FLiveLinkPanelController>& InPanelController)
+		: FWorkflowTabFactory(UE::LiveLinkHub::SourceDetailsTabId, MoveTemp(InHostingApp))
+		, WeakPanelController(InPanelController)
+	{
+		TabLabel = LOCTEXT("SourceDetailsTabLabel", "Source Details");
+		TabIcon = FSlateIcon(FAppStyle::GetAppStyleSetName(), "LevelEditor.Tabs.Details");
+
+		bIsSingleton = true;
+	}
+
+	virtual TSharedRef<SWidget> CreateTabBody(const FWorkflowTabSpawnInfo& Info) const override
+	{
+		TSharedRef<SWidget> SourcesDetails = SNullWidget::NullWidget;
+		if (TSharedPtr<FLiveLinkPanelController> PanelController = WeakPanelController.Pin())
+		{
+			SourcesDetails = PanelController->SourcesDetailsView.ToSharedRef();
+		}
+
+		return SourcesDetails;
+	}
+
+	virtual FText GetTabToolTipText(const FWorkflowTabSpawnInfo& Info) const override
+	{
+		return LOCTEXT("SourceDetailisTabToolTip", "Displays details for the selected LiveLink source.");
+	}
+
+private:
+	TWeakPtr<FLiveLinkPanelController> WeakPanelController;
+};
+
+struct FLiveLinkHubSubjectsTabSummoner : public FWorkflowTabFactory
+{
+public:
+	FLiveLinkHubSubjectsTabSummoner(TSharedPtr<class FAssetEditorToolkit> InHostingApp, const TSharedPtr<FLiveLinkPanelController>& InPanelController)
+		: FWorkflowTabFactory(UE::LiveLinkHub::SubjectsTabId, MoveTemp(InHostingApp))
+		, WeakPanelController(InPanelController)
+	{
+		TabLabel = LOCTEXT("SubjectsTabLabel", "Subjects");
+		TabIcon = FSlateIcon(FAppStyle::GetAppStyleSetName(), "LevelEditor.Tabs.Details");
+
+		bIsSingleton = true;
+	}
+
+	virtual TSharedRef<SWidget> CreateTabBody(const FWorkflowTabSpawnInfo& Info) const override
+	{
+		FLiveLinkHubModule& LiveLinkHubModule = FModuleManager::Get().GetModuleChecked<FLiveLinkHubModule>("LiveLinkHub");
+		TSharedRef<SWidget> SubjectsView = SNullWidget::NullWidget;
+		if (TSharedPtr<FLiveLinkPanelController> PanelController = WeakPanelController.Pin())
+		{
+			SubjectsView = PanelController->SubjectsView->SubjectsTreeView.ToSharedRef();
+		}
+
+		return SNew(SSplitter)
+			.Orientation(EOrientation::Orient_Vertical)
+			+ SSplitter::Slot()
+			.Value(0.5f)
+			[
+				SubjectsView
+			]
+			+
+			SSplitter::Slot()
+			.Value(0.5f)
+			[
+				LiveLinkHubModule.GetSubjectController()->MakeSubjectView()
+			];
+	}
+
+	virtual FText GetTabToolTipText(const FWorkflowTabSpawnInfo& Info) const override
+	{
+		return LOCTEXT("SubjectsTabToolTip", "View the list of Live Link subjects.");
+	}
+
+private:
+	/** Panel controller used to get the list of subjects. */
+	TWeakPtr<FLiveLinkPanelController> WeakPanelController;
+};
+
+struct FLiveLinkHubPlaybackTabSummoner : public FWorkflowTabFactory
+{
+public:
+	FLiveLinkHubPlaybackTabSummoner(TSharedPtr<class FAssetEditorToolkit> InHostingApp)
+		: FWorkflowTabFactory(UE::LiveLinkHub::PlaybackTabId, MoveTemp(InHostingApp))
+	{
+		TabLabel = LOCTEXT("PlaybackTabLabel", "Playback");
+		TabIcon = FSlateIcon(UE::LiveLinkHub::LiveLinkStyleName, TEXT("LiveLinkHub.Playback.Icon"));
+
+		bIsSingleton = true;
+	}
+
+	virtual TSharedRef<SWidget> CreateTabBody(const FWorkflowTabSpawnInfo& Info) const override
+	{
+		FLiveLinkHubModule& LiveLinkHubModule = FModuleManager::Get().GetModuleChecked<FLiveLinkHubModule>("LiveLinkHub");
+
+		return SNew(SVerticalBox)
+			+ SVerticalBox::Slot()
+			.FillHeight(1.0f)
+			[
+				LiveLinkHubModule.GetRecordingListController()->MakeRecordingList()
+			]
+			+ SVerticalBox::Slot()
+			.AutoHeight()
+			[
+				LiveLinkHubModule.GetPlaybackController()->MakePlaybackWidget()
+			];
+	}
+
+	virtual FText GetTabToolTipText(const FWorkflowTabSpawnInfo& Info) const override
+	{
+		return LOCTEXT("PlaybackTabToolTip", "Hosts the LiveLink recording playback functionality.");
+	}
+};
+
+struct FLiveLinkHubClientsTabSummoner : public FWorkflowTabFactory
+{
+public:
+	FLiveLinkHubClientsTabSummoner(TSharedPtr<class FAssetEditorToolkit> InHostingApp)
+		: FWorkflowTabFactory(UE::LiveLinkHub::ClientsTabId, MoveTemp(InHostingApp))
+	{
+		TabLabel = LOCTEXT("ClientsTabLabel", "Clients");
+		TabIcon = FSlateIcon(FAppStyle::GetAppStyleSetName(), "LevelEditor.Tabs.Details");
+
+		bIsSingleton = true;
+	}
+
+	virtual TSharedRef<SWidget> CreateTabBody(const FWorkflowTabSpawnInfo& Info) const override
+	{
+		FLiveLinkHubModule& LiveLinkHubModule = FModuleManager::Get().GetModuleChecked<FLiveLinkHubModule>("LiveLinkHub");
+		TSharedPtr<FLiveLinkHubClientsController> ClientsController = LiveLinkHubModule.GetLiveLinkHub()->GetClientsController();
+
+		return ClientsController->MakeClientsView();
+	}
+
+	virtual FText GetTabToolTipText(const FWorkflowTabSpawnInfo& Info) const override
+	{
+		return LOCTEXT("ClientsTabToolTip", "Displays the list of connected Unreal Editor instances.");
+	}
+};
+
+struct FLiveLinkHubClientDetailsTabSummoner : public FWorkflowTabFactory
+{
+public:
+	FLiveLinkHubClientDetailsTabSummoner(TSharedPtr<class FAssetEditorToolkit> InHostingApp)
+		: FWorkflowTabFactory(UE::LiveLinkHub::ClientDetailsTabId, MoveTemp(InHostingApp))
+	{
+		TabLabel = LOCTEXT("ClientDetailsTabLabel", "Client Details");
+		TabIcon = FSlateIcon(FAppStyle::GetAppStyleSetName(), "LevelEditor.Tabs.Details");
+
+		bIsSingleton = true;
+	}
+
+	virtual TSharedRef<SWidget> CreateTabBody(const FWorkflowTabSpawnInfo& Info) const override
+	{
+		FLiveLinkHubModule& LiveLinkHubModule = FModuleManager::Get().GetModuleChecked<FLiveLinkHubModule>("LiveLinkHub");
+		TSharedPtr<FLiveLinkHubClientsController> ClientsController = LiveLinkHubModule.GetLiveLinkHub()->GetClientsController();
+
+		return ClientsController->MakeClientDetailsView();
+	}
+
+	virtual FText GetTabToolTipText(const FWorkflowTabSpawnInfo& Info) const override
+	{
+		return LOCTEXT("ClientDetailsTabToolTip", "Displays details for the selected client.");
+	}
+};
+
+/** Default application mode for the hub. Hosts the tabs necessary for viewing Sources, Subjects and clients. */
+class FLiveLinkHubCreatorAppMode : public FLiveLinkHubApplicationMode
+{
+public:
+	FLiveLinkHubCreatorAppMode(TSharedPtr<class FWorkflowCentricApplication> App)
+		: FLiveLinkHubApplicationMode("CreatorMode", LOCTEXT("CreatorModeLabel", "Creator Mode"), App)
+	{
+
+		TabLayout = FTabManager::NewLayout("LiveLinkHubCreatorMode_v1.0");
+		const TSharedRef<FTabManager::FArea> MainWindowArea = FTabManager::NewPrimaryArea()->SetOrientation(Orient_Horizontal);
+
+		PanelController = MakeShared<FLiveLinkPanelController>(TAttribute<bool>::CreateRaw(this, &FLiveLinkHubCreatorAppMode::IsSourcePanelReadOnly));
+		PanelController->OnSubjectSelectionChanged().AddRaw(this, &FLiveLinkHubCreatorAppMode::OnSubjectSelectionChanged);
+
+		TabLayout->AddArea
+		(
+			FTabManager::NewPrimaryArea()
+			->SetOrientation(Orient_Vertical)
+			->Split
+			(
+				FTabManager::NewSplitter()
+				->SetSizeCoefficient(1.f)
+				->SetOrientation(Orient_Horizontal)
+				->Split
+				(
+					FTabManager::NewSplitter()
+					->SetSizeCoefficient(0.25f)
+					->SetOrientation(Orient_Vertical)
+					->Split
+					(
+						FTabManager::NewStack()
+						->SetSizeCoefficient(0.5f)
+						->SetHideTabWell(true)
+						->AddTab(UE::LiveLinkHub::SourcesTabId, ETabState::OpenedTab)
+					)
+					->Split
+					(
+						FTabManager::NewStack()
+						->SetSizeCoefficient(0.5f)
+						->AddTab(UE::LiveLinkHub::SourceDetailsTabId, ETabState::OpenedTab)
+					)
+				)
+				->Split
+				(
+					FTabManager::NewSplitter()
+					->SetSizeCoefficient(0.25f)
+					->SetOrientation(Orient_Vertical)
+					->Split
+					(
+						FTabManager::NewStack()
+						->SetSizeCoefficient(0.5f)
+						->AddTab(UE::LiveLinkHub::SubjectsTabId, ETabState::OpenedTab)
+						->AddTab(UE::LiveLinkHub::PlaybackTabId, ETabState::OpenedTab)
+						->SetForegroundTab(UE::LiveLinkHub::SubjectsTabId)
+					)
+				)
+				->Split
+				(
+					FTabManager::NewSplitter()
+					->SetSizeCoefficient(0.25f)
+					->SetOrientation(Orient_Vertical)
+					->Split
+					(
+						FTabManager::NewStack()
+						->SetSizeCoefficient(0.5f)
+						->SetHideTabWell(true)
+						->AddTab(UE::LiveLinkHub::ClientsTabId, ETabState::OpenedTab)
+					)
+					->Split
+					(
+						FTabManager::NewStack()
+						->SetSizeCoefficient(0.5f)
+						->SetHideTabWell(true)
+						->AddTab(UE::LiveLinkHub::ClientDetailsTabId, ETabState::OpenedTab)
+					)
+				)
+			)
+		);
+		
+		ToolbarExtender->AddToolBarExtension("Asset", EExtensionHook::After, App->GetToolkitCommands(), FToolBarExtensionDelegate::CreateRaw(this, &FLiveLinkHubCreatorAppMode::AddRecordingInfoToolbarExtension));
+
+		TabFactories.RegisterFactory(MakeShared<FLiveLinkHubSourcesTabSummoner>(App, PanelController));
+		TabFactories.RegisterFactory(MakeShared<FLiveLinkHubSourceDetailsTabSummoner>(App, PanelController));
+		TabFactories.RegisterFactory(MakeShared<FLiveLinkHubSubjectsTabSummoner>(App, PanelController));
+		TabFactories.RegisterFactory(MakeShared<FLiveLinkHubPlaybackTabSummoner>(App));
+		TabFactories.RegisterFactory(MakeShared<FLiveLinkHubClientsTabSummoner>(App));
+		TabFactories.RegisterFactory(MakeShared<FLiveLinkHubClientDetailsTabSummoner>(App));
+	}
+
+	virtual ~FLiveLinkHubCreatorAppMode()
+	{
+		if (PanelController)
+		{
+			PanelController->OnSubjectSelectionChanged().RemoveAll(this);
+		}
+
+		PanelController.Reset();
+	}
+
+	//~ Begin FLiveLinkHubApplicationMode interface
+	virtual FSlateIcon GetModeIcon() const
+	{
+		return FSlateIcon("LiveLinkStyle", TEXT("LiveLinkHub.Subjects.Icon"));
+	}
+
+
+	virtual TArray<TSharedRef<SWidget>> GetToolbarWidgets_Impl() override
+	{
+		FLiveLinkHubModule& LiveLinkHubModule = FModuleManager::Get().GetModuleChecked<FLiveLinkHubModule>("LiveLinkHub");
+		return {
+			LiveLinkHubModule.GetRecordingController()->MakeRecordToolbarEntry(),
+			SNew(SLiveLinkTimecode)
+		}; 
+	}
+
+	virtual TArray<TSharedRef<SWidget>> GetStatusBarWidgets_Impl() override
+	{
+		UToolMenu* Menu = UToolMenus::Get()->ExtendMenu(TEXT("LiveLinkHub.ToolBar"));
+
+		FToolMenuSection& SessionConfigSection = Menu->AddSection(TEXT("SessionConfig"), FText::GetEmpty(), FToolMenuInsert(NAME_None, EToolMenuInsertType::First));
+
+		return {
+			SNew(SBox)
+			.Padding(0, 4.0, 0, 0)
+			[
+				SNew(STextBlock)
+				.Text_Raw(this, &FLiveLinkHubCreatorAppMode::GetLoadedConfigText)
+			]
+		};
+	}
+
+	virtual void RegisterTabFactories(TSharedPtr<FTabManager> InTabManager) override
+	{
+		if (TSharedPtr<FWorkflowCentricApplication> App = WeakApp.Pin())
+		{
+			// Pushing the tab factories will register the tab spwaners for the current mode.
+			App->PushTabFactories(TabFactories);
+		}
+
+		FApplicationMode::RegisterTabFactories(InTabManager);
+	}
+	//~ End FLiveLinkHubApplicationMode interface
+
+private:
+	/** Handles updating the active subject in the subject controller. */
+	void OnSubjectSelectionChanged(const FLiveLinkSubjectKey& SubjectKey)
+	{
+		FLiveLinkHubModule& LiveLinkHubModule = FModuleManager::Get().GetModuleChecked<FLiveLinkHubModule>("LiveLinkHub");
+		LiveLinkHubModule.GetSubjectController()->SetSubject(SubjectKey);
+	}
+
+	/** Returns whether the source panel is read only (such as when we're playing back a recording) */
+	bool IsSourcePanelReadOnly() const
+	{
+		FLiveLinkHubModule& LiveLinkHubModule = FModuleManager::Get().GetModuleChecked<FLiveLinkHubModule>("LiveLinkHub");
+		return LiveLinkHubModule.GetPlaybackController()->IsInPlayback() || LiveLinkHubModule.GetRecordingController()->IsRecording();
+	}
+
+	/** Returns the name of the active session. */
+	FText GetLoadedConfigText() const
+	{
+		if (TSharedPtr<ILiveLinkHubSessionManager> SessionManager = FModuleManager::Get().GetModuleChecked<FLiveLinkHubModule>("LiveLinkHub").GetLiveLinkHub()->GetSessionManager())
+		{
+			const FString FileName = FPaths::GetBaseFilename(SessionManager->GetLastConfigPath());
+			return FileName.IsEmpty() ? LOCTEXT("UntitledConfig", "Untitled") : FText::FromString(FileName);
+		}
+
+		return FText::GetEmpty();
+	}
+
+	/** Adds the take and slate widgets to the toolbar.  */
+	void AddRecordingInfoToolbarExtension(FToolBarBuilder& Builder)
+	{
+		Builder.AddWidget(
+			SNew(SHorizontalBox)
+			+ SHorizontalBox::Slot()
+			.FillWidth(1.0)
+			.VAlign(VAlign_Fill)
+			[
+				SNew(SLiveLinkRecordingSessionInfo)
+			], /*TutorialHighlightName*/NAME_None, /*bSearchable*/ false, EHorizontalAlignment::HAlign_Right
+		);
+	}
+
+
+protected:
+	/** Set of spawnable tabs in the mode */
+	FWorkflowAllowedTabSet TabFactories;
+
+	/** Holds the livelink panel controller responsible for creating sources and subjects tabs. */
+	TSharedPtr<FLiveLinkPanelController> PanelController;
+};
+
+#undef LOCTEXT_NAMESPACE
diff --git a/Engine/Plugins/Animation/LiveLinkHub/Source/LiveLinkHub/Private/UI/Widgets/SApplicationModeSwitcher.h b/Engine/Plugins/Animation/LiveLinkHub/Source/LiveLinkHub/Private/UI/Widgets/SApplicationModeSwitcher.h
new file mode 100644
--- /dev/null
+++ b/Engine/Plugins/Animation/LiveLinkHub/Source/LiveLinkHub/Private/UI/Widgets/SApplicationModeSwitcher.h
@@ -0,0 +1,114 @@
+// Copyright Epic Games, Inc. All Rights Reserved.
+
+#pragma once
+
+#include "Widgets/SCompoundWidget.h"
+
+#include "LiveLinkHub.h"
+#include "Widgets/Images/SImage.h"
+#include "Widgets/Input/SComboButton.h"
+
+
+class SApplicationModeSwitcher : public SCompoundWidget
+{
+public:
+	SLATE_BEGIN_ARGS(SApplicationModeSwitcher)
+	{}
+	SLATE_END_ARGS()
+
+	void Construct(const FArguments& InArgs)
+	{
+		ChildSlot
+		[
+			SNew(SComboButton)
+			.ContentPadding(FMargin(2.0f, 3.0f, 2.0f, 3.0f))
+			.OnGetMenuContent(this, &SApplicationModeSwitcher::GetModeSwitcherContent)
+			.ButtonContent()
+			[
+				SNew(SHorizontalBox)
+				+ SHorizontalBox::Slot()
+				.AutoWidth()
+				[
+					SAssignNew(ActiveModeImage, SImage)
+					.ColorAndOpacity(FSlateColor::UseForeground())
+				]
+				+ SHorizontalBox::Slot()
+				.Padding(FMargin(4.0f, 0.0f, 0.0f, 0.0f))
+				.AutoWidth()
+				[
+					SAssignNew(ActiveModeDisplayName, STextBlock)
+				]
+			]
+		];
+
+		if (TSharedPtr<FLiveLinkHub> LiveLinkHub = FLiveLinkHub::Get())
+		{
+			LiveLinkHub->OnApplicationModeChanged().AddSP(this, &SApplicationModeSwitcher::OnAppModeChanged);
+
+			if (TOptional<FLiveLinkHubAppModeInfo> AppModeInfo = LiveLinkHub->GetModeInfo(LiveLinkHub->GetCurrentMode()))
+			{
+				ActiveModeImage->SetImage(AppModeInfo->Icon.GetIcon());
+				ActiveModeDisplayName->SetText(AppModeInfo->DisplayName);
+			}
+		}
+	}
+
+	virtual ~SApplicationModeSwitcher()
+	{
+		if (TSharedPtr<FLiveLinkHub> LiveLinkHub = FLiveLinkHub::Get())
+		{
+			LiveLinkHub->OnApplicationModeChanged().RemoveAll(this);
+		}
+	}
+
+private:
+	TSharedRef<SWidget> GetModeSwitcherContent()
+	{
+		FMenuBuilder MenuBuilder(true, NULL);
+		if (TSharedPtr<FLiveLinkHub> LiveLinkHub = FLiveLinkHub::Get())
+		{
+			TArray<FName> Modes = LiveLinkHub->GetApplicationModes();
+			FName CurrentModeName = LiveLinkHub->GetCurrentMode();
+
+			for (FName Mode : Modes)
+			{
+				if (Mode != CurrentModeName)
+				{
+					if (TOptional<FLiveLinkHubAppModeInfo> AppModeInfo = LiveLinkHub->GetModeInfo(Mode))
+					{
+						// Tooltip text, Icon, localized name
+						MenuBuilder.AddMenuEntry(AppModeInfo->DisplayName, AppModeInfo->DisplayName, AppModeInfo->Icon, FExecuteAction::CreateRaw(this, &SApplicationModeSwitcher::SetActiveMode, Mode));
+					}
+				}
+			}
+		}
+
+		return MenuBuilder.MakeWidget();
+	}
+
+	void SetActiveMode(FName ModeName)
+	{
+		if (TSharedPtr<FLiveLinkHub> LiveLinkHub = FLiveLinkHub::Get())
+		{
+			LiveLinkHub->SetCurrentMode(ModeName);
+		}
+	}
+
+	void OnAppModeChanged(FName NewModeName) const
+	{
+		if (TSharedPtr<FLiveLinkHub> LiveLinkHub = FLiveLinkHub::Get())
+		{
+			if (TOptional<FLiveLinkHubAppModeInfo> AppModeInfo = LiveLinkHub->GetModeInfo(NewModeName))
+			{
+				ActiveModeImage->SetImage(AppModeInfo->Icon.GetIcon());
+				ActiveModeDisplayName->SetText(AppModeInfo->DisplayName);
+			}
+		}
+	}
+
+private:
+
+	TSharedPtr<SImage> ActiveModeImage;
+	TSharedPtr<STextBlock> ActiveModeDisplayName;
+};
+
diff --git a/Engine/Plugins/Animation/LiveLinkHub/Source/LiveLinkHub/Private/UI/Widgets/SLiveLinkHubSettings.h b/Engine/Plugins/Animation/LiveLinkHub/Source/LiveLinkHub/Private/UI/Widgets/SLiveLinkHubSettings.h
--- a/Engine/Plugins/Animation/LiveLinkHub/Source/LiveLinkHub/Private/UI/Widgets/SLiveLinkHubSettings.h
+++ b/Engine/Plugins/Animation/LiveLinkHub/Source/LiveLinkHub/Private/UI/Widgets/SLiveLinkHubSettings.h
@@ -44,40 +44,35 @@ public:
 		ISettingsEditorModule* SettingsEditorModule = static_cast<ISettingsEditorModule*>(FModuleManager::Get().LoadModule("SettingsEditor"));
 		SettingsEditorModule->SetShouldRegisterSettingCallback(FShouldRegisterSettingsDelegate::CreateLambda([](UDeveloperSettings* Settings) { return false; }));
 
+		// Unregister in case the widget was re-created.
+		FGlobalTabmanager::Get()->UnregisterTabSpawner(SettingsTabId);
+
 		FGlobalTabmanager::Get()->RegisterTabSpawner(SettingsTabId, FOnSpawnTab::CreateRaw(this, &SLiveLinkHubSettings::SpawnSettingsTab))
 			.SetDisplayName(LOCTEXT("LiveLinkHubSettingsTabLabel", "Settings"))
 			.SetIcon(FSlateIcon(FAppStyle::GetAppStyleSetName(), "LevelEditor.Tabs.Details"));
 
 		ChildSlot
-		[
-			SNew(SComboButton)
-			.ContentPadding(FMargin(2.0f, 4.0f))
-			.ButtonStyle(&FAppStyle::Get().GetWidgetStyle<FButtonStyle>("Button"))
-			.ForegroundColor(FSlateColor::UseStyle())
-			.ButtonContent()
 			[
-				SNew(SHorizontalBox)
-				+ SHorizontalBox::Slot()
-				.VAlign(VAlign_Center)
-				.Padding(4.0f, 0.0f)
-				.AutoWidth()
-				[
-					SNew(STextBlock)
-					.Font(FAppStyle::Get().GetFontStyle("FontAwesome.14"))
-					.Text(FEditorFontGlyphs::Cogs)
-				]
-				+ SHorizontalBox::Slot()
-				.VAlign(VAlign_Center)
-				.Padding(2.0f, 0, 0, 0)
-				.AutoWidth()
-				[
-					SNew(STextBlock)
-					.TextStyle(FAppStyle::Get(), "SmallButtonText")
-					.Text(LOCTEXT("SettingsLabel", "Settings"))
-				]
-			]
-			.OnGetMenuContent(this, &SLiveLinkHubSettings::OnSettingsComboBoxOpen)
-		];
+				SNew(SComboButton)
+					.ContentPadding(FMargin(0, 4.0f))
+					.ButtonStyle(&FAppStyle::Get().GetWidgetStyle<FButtonStyle>("Button"))
+					.ForegroundColor(FSlateColor::UseStyle())
+					.ToolTipText(LOCTEXT("SettingsToolTip", "Settings"))
+					.ButtonContent()
+					[
+						SNew(SHorizontalBox)
+							+ SHorizontalBox::Slot()
+							.VAlign(VAlign_Center)
+							.Padding(4.0f, 0.0f)
+							.AutoWidth()
+							[
+								SNew(STextBlock)
+									.Font(FAppStyle::Get().GetFontStyle("FontAwesome.14"))
+									.Text(FEditorFontGlyphs::Cogs)
+							]
+					]
+					.OnGetMenuContent(this, &SLiveLinkHubSettings::OnSettingsComboBoxOpen)
+			];
 	}
 
 	~SLiveLinkHubSettings()
@@ -148,9 +143,9 @@ public:
 			LOCTEXT("PluginsMenuEntry_ToolTip", "Open the Plugins Browser tab."),
 			FSlateIcon("PluginStyle", "Plugins.TabIcon"),
 			FExecuteAction::CreateStatic([]()
-			{
-				FGlobalTabmanager::Get()->TryInvokeTab(FName("PluginsEditor"));
-			}));
+				{
+					FGlobalTabmanager::Get()->TryInvokeTab(FName("PluginsEditor"));
+				}));
 		MenuBuilder.EndSection();
 
 		return MenuBuilder.MakeWidget();
diff --git a/Engine/Plugins/Animation/LiveLinkHub/Source/LiveLinkHub/Private/UI/Window/LiveLinkHubWindowController.cpp b/Engine/Plugins/Animation/LiveLinkHub/Source/LiveLinkHub/Private/UI/Window/LiveLinkHubWindowController.cpp
--- a/Engine/Plugins/Animation/LiveLinkHub/Source/LiveLinkHub/Private/UI/Window/LiveLinkHubWindowController.cpp
+++ b/Engine/Plugins/Animation/LiveLinkHub/Source/LiveLinkHub/Private/UI/Window/LiveLinkHubWindowController.cpp
@@ -7,6 +7,7 @@
 #include "Framework/Application/SlateApplication.h"
 #include "Framework/Docking/LayoutService.h"
 #include "Framework/Notifications/NotificationManager.h"
+#include "LiveLinkHub.h"
 #include "HAL/PlatformApplicationMisc.h"
 #include "Misc/ConfigCacheIni.h"
 #include "StandaloneRenderer.h"
@@ -15,9 +16,8 @@
 
 #define LOCTEXT_NAMESPACE "LiveLinkHubWindowController"
 
-FLiveLinkHubWindowController::FLiveLinkHubWindowController(const FLiveLinkHubWindowInitParams& Params)
-	: LiveLinkHubLayoutIni(Params.LiveLinkHubLayoutIni)
-	, MainTabController(MakeShared<FLiveLinkHubMainTabController>())
+FLiveLinkHubWindowController::FLiveLinkHubWindowController()
+	: MainTabController(MakeShared<FLiveLinkHubMainTabController>())
 {
 	ModalWindowManager = InitializeSlateApplication();
 }
@@ -30,10 +30,6 @@ FLiveLinkHubWindowController::~FLiveLinkHubWindowController()
 	{
 		RootWindow->SetOnWindowClosed(nullptr);
 	}
-
-#if IS_PROGRAM
-	FSlateApplication::Shutdown();
-#endif
 }
 
 TSharedRef<SWindow> FLiveLinkHubWindowController::CreateWindow()
@@ -68,30 +64,29 @@ TSharedRef<SWindow> FLiveLinkHubWindowController::CreateWindow()
 
 	return RootWindowRef;
 }
-void FLiveLinkHubWindowController::RestoreLayout()
+void FLiveLinkHubWindowController::RestoreLayout(TSharedPtr<FAssetEditorToolkit> AssetEditorToolkit)
 {
 	TRACE_CPUPROFILER_EVENT_SCOPE(FLiveLinkHubWindowController::RestoreLayout);
 
-	const TSharedRef<FTabManager::FLayout> DefaultLayout = FTabManager::NewLayout("LiveLinkHub_v1.0");
-	const TSharedRef<FTabManager::FArea> MainWindowArea = FTabManager::NewPrimaryArea()->SetOrientation(Orient_Horizontal);
-	const TSharedRef<FTabManager::FStack> MainStack = FTabManager::NewStack();
-	InitComponents(MainStack);
-	MainStack->SetHideTabWell(true);
-
-	MainWindowArea->Split(MainStack);
-	DefaultLayout->AddArea(MainWindowArea);
-	PersistentLayout = DefaultLayout;
-
-	{
-		TRACE_CPUPROFILER_EVENT_SCOPE(FLiveLinkHubWindowController::LoadFromConfig);
-		PersistentLayout = FLayoutSaveRestore::LoadFromConfig(LiveLinkHubLayoutIni, DefaultLayout);
-	}
-
+	const FName LayoutName = TEXT("LiveLinkHub_v1.1");
+	const TSharedRef<FTabManager::FLayout> DefaultLayout =
+		FTabManager::NewLayout(LayoutName)
+		->AddArea
+		(
+			// Toolkits window
+			FTabManager::NewPrimaryArea()
+			->Split
+			(
+				FTabManager::NewStack()
+				->SetSizeCoefficient(1.0f)
+				->AddTab("StandaloneToolkit", ETabState::ClosedTab)
+			)
+		);
 
 	{
 		TRACE_CPUPROFILER_EVENT_SCOPE(FLiveLinkHubWindowController::RestoreFrom);
 		constexpr bool bEmbedTitleAreaContent = true;
-		const TSharedPtr<SWidget> Content = FGlobalTabmanager::Get()->RestoreFrom(PersistentLayout.ToSharedRef(), RootWindow, bEmbedTitleAreaContent, EOutputCanBeNullptr::Never);
+		const TSharedPtr<SWidget> Content = FGlobalTabmanager::Get()->RestoreFrom(DefaultLayout, RootWindow, bEmbedTitleAreaContent, EOutputCanBeNullptr::Never);
 		RootWindow->SetContent(Content.ToSharedRef());
 	}
 
@@ -99,24 +94,17 @@ void FLiveLinkHubWindowController::RestoreLayout()
 	constexpr bool bForceWindowToFront = true;
 	RootWindow->BringToFront(bForceWindowToFront);
 
-	MainTabController->OpenTab();
+	// Pass a dummy object to the asset editor since we're not actually editing an object.
+	UObject* DummyObject = GetMutableDefault<UObject>();
+	AssetEditorToolkit->InitAssetEditor(EToolkitMode::Standalone, nullptr, "LiveLinkHub", DefaultLayout, /*bCreateDefaultStandaloneMenu*/true, /*bCreateDefaultToolbar*/true, DummyObject, /*bInIsToolbarFocusable*/true, true, {});
+
+	FGlobalTabmanager::Get()->SetMainTab(FTabId("StandaloneToolkit"));
 }
 
 TSharedPtr<FModalWindowManager> FLiveLinkHubWindowController::InitializeSlateApplication()
 {
 	TRACE_CPUPROFILER_EVENT_SCOPE(FLiveLinkHubWindowController::InitializeSlateApplication);
 
-#if IS_PROGRAM
-	FSlateApplication::InitializeAsStandaloneApplication(GetStandardStandaloneRenderer());
-
-	// @hack This call will silently fail since we're running a commandlet, so pretend like we aren't one for it.
-	// In the future the slate code should be changed to allow visual commandlets to enable high dpi mode.
-	const bool bIsRunningCommandlet = PRIVATE_GIsRunningCommandlet;
-	PRIVATE_GIsRunningCommandlet = false;
-	FSlateApplication::InitHighDPI(true);
-	PRIVATE_GIsRunningCommandlet = bIsRunningCommandlet;
-#endif
-
 	const FText ApplicationTitle = LOCTEXT("AppTitle", "Live Link Hub");
 	FGlobalTabmanager::Get()->SetApplicationTitle(ApplicationTitle);
 
@@ -137,21 +125,15 @@ void FLiveLinkHubWindowController::InitComponents(const TSharedRef<FTabManager::
 
 void FLiveLinkHubWindowController::OnWindowClosed(const TSharedRef<SWindow>& Window)
 {
-	SaveLayout();
+	if (TSharedPtr<FLiveLinkHub> LiveLinkHub = FLiveLinkHub::Get())
+	{
+		LiveLinkHub->OnClose();
+	}
+
 	RootWindow.Reset();
 
-#if !IS_PROGRAM
 	RequestEngineExit(TEXT("FLiveLinkHubWindowController::OnWindowClosed"));
-#endif
 }
 
-void FLiveLinkHubWindowController::SaveLayout() const
-{
-	if (PersistentLayout)
-	{
-		FLayoutSaveRestore::SaveToConfig(LiveLinkHubLayoutIni, PersistentLayout.ToSharedRef());
-	    GConfig->Flush(false, LiveLinkHubLayoutIni);
-	}
-}
 
 #undef LOCTEXT_NAMESPACE 
diff --git a/Engine/Plugins/Animation/LiveLinkHub/Source/LiveLinkHub/Private/UI/Window/LiveLinkHubWindowController.h b/Engine/Plugins/Animation/LiveLinkHub/Source/LiveLinkHub/Private/UI/Window/LiveLinkHubWindowController.h
--- a/Engine/Plugins/Animation/LiveLinkHub/Source/LiveLinkHub/Private/UI/Window/LiveLinkHubWindowController.h
+++ b/Engine/Plugins/Animation/LiveLinkHub/Source/LiveLinkHub/Private/UI/Window/LiveLinkHubWindowController.h
@@ -4,6 +4,7 @@
 
 #include "Framework/Docking/TabManager.h"
 
+class FAssetEditorToolkit;
 class FLiveLinkHubMainTabController;
 class FLiveLinkHubSessionTabBase;
 class FOutputLogController;
@@ -12,30 +13,18 @@ class ILiveLinkHubComponent;
 class FModalWindowManager;
 class SWindow;
 
-struct FLiveLinkHubWindowInitParams
-{
-	/** Config path for hub layout ini */
-	FString LiveLinkHubLayoutIni;
-
-	FLiveLinkHubWindowInitParams() = default;
-
-	FLiveLinkHubWindowInitParams(FString LiveLinkHubLayoutIni)
-		: LiveLinkHubLayoutIni(MoveTemp(LiveLinkHubLayoutIni))
-	{}
-};
-
 /** Responsible for creating the Slate window for the hub. */
 class FLiveLinkHubWindowController : public TSharedFromThis<FLiveLinkHubWindowController>
 {
 public:
 	
-	FLiveLinkHubWindowController(const FLiveLinkHubWindowInitParams& Params);
+	FLiveLinkHubWindowController();
 	~FLiveLinkHubWindowController();
 
 	/** Get the root window. */
 	TSharedPtr<SWindow> GetRootWindow() const { return RootWindow; }
 	/** Restore the window's layout from a config. */
-	void RestoreLayout();
+	void RestoreLayout(TSharedPtr<FAssetEditorToolkit> AssetEditorToolkit);
 private:
 	/** Create the main window. */
 	TSharedRef<SWindow> CreateWindow();
@@ -45,8 +34,6 @@ private:
 	void InitComponents(const TSharedRef<FTabManager::FStack>& MainArea);
 	/** Window closed handler. */
 	void OnWindowClosed(const TSharedRef<SWindow>& Window);
-	/** Save the layout to a json file. */
-	void SaveLayout() const;
 private:
 	/** Pointer to the livelink client. */
 	ILiveLinkClient* LiveLinkClient = nullptr;
diff --git a/Engine/Plugins/Animation/LiveLinkHub/Source/LiveLinkHub/Public/ILiveLinkHubModule.h b/Engine/Plugins/Animation/LiveLinkHub/Source/LiveLinkHub/Public/ILiveLinkHubModule.h
--- a/Engine/Plugins/Animation/LiveLinkHub/Source/LiveLinkHub/Public/ILiveLinkHubModule.h
+++ b/Engine/Plugins/Animation/LiveLinkHub/Source/LiveLinkHub/Public/ILiveLinkHubModule.h
@@ -6,6 +6,7 @@
 
 class LIVELINKHUB_API ILiveLinkHubModule : public IModuleInterface
 {
+public:
 	UE_DEPRECATED(5.6, "There's no need to call this method anymore, it's now handled by StartupModule and ShutdownModule.")
 	virtual void PreinitializeLiveLinkHub() {}
 
diff --git a/Engine/Plugins/Animation/LiveLinkHub/Source/LiveLinkHub/Public/LiveLinkHubApplicationMode.h b/Engine/Plugins/Animation/LiveLinkHub/Source/LiveLinkHub/Public/LiveLinkHubApplicationMode.h
new file mode 100644
--- /dev/null
+++ b/Engine/Plugins/Animation/LiveLinkHub/Source/LiveLinkHub/Public/LiveLinkHubApplicationMode.h
@@ -0,0 +1,62 @@
+// Copyright Epic Games, Inc. All Rights Reserved.
+
+#pragma once
+
+#include "Features/IModularFeature.h"
+
+#include "WorkflowOrientedApp/ApplicationMode.h"
+
+struct FSlateIcon;
+class FToolBarBuilder;
+class FWorkflowCentricApplication;
+class SWidget;
+
+/** Factory used by external plugins to register custom LiveLinkHub layouts. */
+class LIVELINKHUB_API ILiveLinkHubApplicationModeFactory : public IModularFeature
+{
+public:
+	/** Name of the modular feature. */
+	static const FName ModularFeatureName;
+
+	/** Instantiate an application mode so LiveLinkHub can register it and display it in its Layout Selector. */
+	virtual TSharedRef<class FLiveLinkHubApplicationMode> CreateApplicationMode(TSharedPtr<FWorkflowCentricApplication> InApp) = 0;
+};
+
+class LIVELINKHUB_API FLiveLinkHubApplicationMode : public FApplicationMode
+{
+public:
+	/** Name of the File menu extension point. Can be used to insert menus before/after in the main menu bar. */
+	static const FName FileMenuExtensionPoint;
+
+	/** Get the icon for this mode. */
+	virtual FSlateIcon GetModeIcon() const = 0;
+
+	/** Get the display name for this mode */
+	const FText& GetDisplayName() const { return DisplayName; }
+
+	/** Method used to gather the toolbar widgets that should be visible when this mode is active. */
+	TArray<TSharedRef<SWidget>> GetToolbarWidgets();
+
+protected:
+	FLiveLinkHubApplicationMode(FName InApplicationMode, FText InDisplayName, TSharedPtr<FWorkflowCentricApplication> InApp);
+
+	//~ Begin FApplicationMode interface
+	virtual void PostActivateMode() override;
+	virtual void PreDeactivateMode() override;
+	//~ End FApplicationMode interface
+
+	/** Implement in a child class in order to have status bar extensions show up when the mode is active. */
+	virtual TArray<TSharedRef<SWidget>> GetToolbarWidgets_Impl() { return {}; };
+
+	
+	/** Get widgets that should appear in the status bar. */
+	virtual TArray<TSharedRef<SWidget>> GetStatusBarWidgets_Impl() { return {}; }
+
+protected:
+	/** The application that hosts this app mode. */
+	TWeakPtr<FWorkflowCentricApplication> WeakApp;
+
+private:
+	/** Display name for this mode. */
+	FText DisplayName;
+};
\ No newline at end of file
