Fix CIS error caused by misssing include guard in LiveLinkHubApplication

#rnx

diff --git a/Engine/Plugins/Animation/LiveLinkHub/Source/LiveLinkHub/Private/LiveLinkHubApplication.cpp b/Engine/Plugins/Animation/LiveLinkHub/Source/LiveLinkHub/Private/LiveLinkHubApplication.cpp
--- a/Engine/Plugins/Animation/LiveLinkHub/Source/LiveLinkHub/Private/LiveLinkHubApplication.cpp
+++ b/Engine/Plugins/Animation/LiveLinkHub/Source/LiveLinkHub/Private/LiveLinkHubApplication.cpp
@@ -2,8 +2,6 @@
 
 #include "LiveLinkHubApplication.h"
 
-#include "ISessionFrontendModule.h"
-#include "ISlateReflectorModule.h"
 #include "LiveLinkHubCreatorAppMode.h"
 #include "Misc/EngineBuildSettings.h"
 #include "StatusBarSubsystem.h"
@@ -13,6 +11,13 @@
 #include "WorkflowOrientedApp/ApplicationMode.h"
 #include "WorkflowOrientedApp/WorkflowTabManager.h"
 
+#if !UE_BUILD_SHIPPING
+#	if WITH_UNREAL_TARGET_DEVELOPER_TOOLS
+#		include "ISessionFrontendModule.h"
+#	endif
+#	include "ISlateReflectorModule.h"
+#endif
+
 #define LOCTEXT_NAMESPACE "LiveLinkHubApplication"
 
 FLiveLinkHubApplication::FLiveLinkHubApplication()
@@ -165,6 +170,7 @@ void FLiveLinkHubApplication::RegisterMainMenu()
 
 void FLiveLinkHubApplication::CreateToolsMenu(FMenuBuilder& MenuBuilder, UToolMenu*) const
 {
+#if !UE_BUILD_SHIPPING
 	MenuBuilder.BeginSection("Development", LOCTEXT("DevelopmentHeader", "Development"));
 
 #if WITH_UNREAL_TARGET_DEVELOPER_TOOLS
@@ -192,6 +198,8 @@ void FLiveLinkHubApplication::CreateToolsMenu(FMenuBuilder& MenuBuilder, UToolMe
 			})
 	);
 	MenuBuilder.EndSection();
+
+#endif /*UE_BUILD_SHIPPING*/
 }
 
 void FLiveLinkHubApplication::FillWindowMenu(FMenuBuilder& MenuBuilder, UToolMenu*)
