Fix missing settings from LiveLinkHub

#rnx

diff --git a/Engine/Plugins/Animation/LiveLinkHub/Source/LiveLinkHub/Private/Settings/LiveLinkSettingsCustomization.cpp b/Engine/Plugins/Animation/LiveLinkHub/Source/LiveLinkHub/Private/Settings/LiveLinkSettingsCustomization.cpp
--- a/Engine/Plugins/Animation/LiveLinkHub/Source/LiveLinkHub/Private/Settings/LiveLinkSettingsCustomization.cpp
+++ b/Engine/Plugins/Animation/LiveLinkHub/Source/LiveLinkHub/Private/Settings/LiveLinkSettingsCustomization.cpp
@@ -14,7 +14,5 @@ TSharedRef<IDetailCustomization> FLiveLinkSettingsCustomization::MakeInstance()
 void FLiveLinkSettingsCustomization::CustomizeDetails(IDetailLayoutBuilder& DetailBuilder)
 {
 	// These properties aren't needed when accessed through Live Link Hub.
-	// Additionally, the Browse To button isn't supported within the hub, and can lead to crashes.
-	DetailBuilder.HideProperty(ULiveLinkSettings::GetDefaultRoleSettingsPropertyName());
 	DetailBuilder.HideProperty(GET_MEMBER_NAME_CHECKED(ULiveLinkSettings, DefaultLiveLinkPreset));
 }
