LiveLink - Only add discovered LiveLinkHub sources if they had the setting enabled

diff --git a/Engine/Plugins/Animation/LiveLink/Source/LiveLinkEditor/Private/LiveLinkEditorModule.cpp b/Engine/Plugins/Animation/LiveLink/Source/LiveLinkEditor/Private/LiveLinkEditorModule.cpp
--- a/Engine/Plugins/Animation/LiveLink/Source/LiveLinkEditor/Private/LiveLinkEditorModule.cpp
+++ b/Engine/Plugins/Animation/LiveLink/Source/LiveLinkEditor/Private/LiveLinkEditorModule.cpp
@@ -110,6 +110,8 @@ public:
 		StyleSet->Set("LiveLinkHub.Playback.Icon", new IMAGE_PLUGIN_BRUSH_SVG(TEXT("Starship/Playback"), Icon16x16));
 		StyleSet->Set("LiveLinkHub.Subjects.Icon", new IMAGE_PLUGIN_BRUSH_SVG(TEXT("Starship/Subjects"), Icon16x16));
 
+		StyleSet->Set("LiveLinkHub.Autodiscover", new IMAGE_PLUGIN_BRUSH_SVG("Starship/AutoDiscover", Icon16x16));
+
 		StyleSet->Set("LiveLinkHub.Icon.Small", new IMAGE_PLUGIN_BRUSH(TEXT("Starship/LiveLinkHub_16x16"), Icon16x16));
 		StyleSet->Set("LiveLinkHub.Icon", new IMAGE_PLUGIN_BRUSH_SVG(TEXT("Starship/LiveLinkHub_Icon"), Icon40x40));
 
diff --git a/Engine/Plugins/Animation/LiveLinkHub/Source/LiveLinkHub/Private/Clients/LiveLinkHubProvider.cpp b/Engine/Plugins/Animation/LiveLinkHub/Source/LiveLinkHub/Private/Clients/LiveLinkHubProvider.cpp
--- a/Engine/Plugins/Animation/LiveLinkHub/Source/LiveLinkHub/Private/Clients/LiveLinkHubProvider.cpp
+++ b/Engine/Plugins/Animation/LiveLinkHub/Source/LiveLinkHub/Private/Clients/LiveLinkHubProvider.cpp
@@ -452,7 +452,11 @@ TArray<FLiveLinkHubClientId> FLiveLinkHubProvider::GetSessionClients() const
 
 TMap<FName, FString> FLiveLinkHubProvider::GetAnnotations() const
 {
-	return Annotations;
+	TMap<FName, FString> AnnotationsCopy = Annotations;
+	
+	// Make a local copy of the annotations map and add the autodiscover annotation since this value might change over time.
+	AnnotationsCopy.Add(FLiveLinkHubMessageAnnotation::AutoDiscoverAnnotation, LexToString(GetDefault<ULiveLinkHubSettings>()->bAutoAddDiscoveredClients));
+	return AnnotationsCopy;
 }
 
 TArray<FLiveLinkHubClientId> FLiveLinkHubProvider::GetDiscoveredClients() const
diff --git a/Engine/Plugins/Animation/LiveLinkHub/Source/LiveLinkHub/Private/Clients/SLiveLinkHubClientsView.h b/Engine/Plugins/Animation/LiveLinkHub/Source/LiveLinkHub/Private/Clients/SLiveLinkHubClientsView.h
--- a/Engine/Plugins/Animation/LiveLinkHub/Source/LiveLinkHub/Private/Clients/SLiveLinkHubClientsView.h
+++ b/Engine/Plugins/Animation/LiveLinkHub/Source/LiveLinkHub/Private/Clients/SLiveLinkHubClientsView.h
@@ -344,6 +344,23 @@ public:
 					.Font( DEFAULT_FONT( "Regular", 14 ) )
 					.Text(LOCTEXT("ClientsHeaderText", "Clients"))
 				]
+				+ SHorizontalBox::Slot()
+				.HAlign(HAlign_Right)
+				.AutoWidth()
+				.Padding(FMargin(4.0, 2.0))
+				[
+					SNew(SCheckBox)
+						.Padding(4.0f)
+						.ToolTipText(LOCTEXT("ToggleAutoDiscoveryToolTip", "Automatically add LiveLinkClients that were discovered."))
+						.Style(FAppStyle::Get(), "ToggleButtonCheckbox")
+						.IsChecked(this, &SLiveLinkHubClientsView::GetAutodiscoveryToggleState)
+						.OnCheckStateChanged(this, &SLiveLinkHubClientsView::ToggleAutodiscovery)
+						[
+							SNew(SImage)
+								.Image(FSlateIcon("LiveLinkStyle", "LiveLinkHub.Autodiscover").GetIcon())
+						]
+				]
+				
 				+ SHorizontalBox::Slot()
 				.HAlign(HAlign_Right)
 				.AutoWidth()
@@ -811,6 +828,20 @@ private:
 		}
 	}
 
+	/** Get the autodiscover checkbox state. */
+	ECheckBoxState GetAutodiscoveryToggleState() const
+	{
+		return GetDefault<ULiveLinkHubSettings>()->bAutoAddDiscoveredClients ? ECheckBoxState::Checked : ECheckBoxState::Unchecked;
+	}
+
+	/** Enable / Disable UE client autodiscovery. */
+	void ToggleAutodiscovery(ECheckBoxState NewState)
+	{
+		GetMutableDefault<ULiveLinkHubSettings>()->bAutoAddDiscoveredClients = NewState == ECheckBoxState::Checked;
+		GetMutableDefault<ULiveLinkHubSettings>()->SaveConfig();
+	}
+
+
 private:
 	/** Delegate called when a client is selected. */
 	FOnClientSelected OnClientSelectedDelegate;
diff --git a/Engine/Plugins/Animation/LiveLinkHub/Source/LiveLinkHubMessaging/Private/LiveLinkHubConnectionManager.h b/Engine/Plugins/Animation/LiveLinkHub/Source/LiveLinkHubMessaging/Private/LiveLinkHubConnectionManager.h
--- a/Engine/Plugins/Animation/LiveLinkHub/Source/LiveLinkHubMessaging/Private/LiveLinkHubConnectionManager.h
+++ b/Engine/Plugins/Animation/LiveLinkHub/Source/LiveLinkHubMessaging/Private/LiveLinkHubConnectionManager.h
@@ -109,14 +109,17 @@ private:
 			TArray<FProviderPollResultPtr> PollResults = ILiveLinkModule::Get().GetMessageBusDiscoveryManager().GetDiscoveryResults();
 			for (const FProviderPollResultPtr& PollResult : PollResults)
 			{
-				if (PollResult->Address == LastAddedSource.Key)
+				const FString* ProviderType = PollResult->Annotations.Find(FLiveLinkHubMessageAnnotation::ProviderTypeAnnotation);
+		
+				// Before this annotation was added, LiveLinkHub would automatically be added, so we need to keep the previous behavior if we discovered an older LiveLinkHub instance.
+				bool bAutoDiscoverLiveLinkHub = true;
+
+				if (const FString* AutoDiscoverAnnotation = PollResult->Annotations.Find(FLiveLinkHubMessageAnnotation::AutoDiscoverAnnotation))
 				{
-					continue;
+					LexFromString(bAutoDiscoverLiveLinkHub, **AutoDiscoverAnnotation);
 				}
 
-				const FString* ProviderType = PollResult->Annotations.Find(FLiveLinkHubMessageAnnotation::ProviderTypeAnnotation);
-
-				if (ProviderType && *ProviderType == UE::LiveLinkHub::Private::LiveLinkHubProviderType)
+				if (ProviderType && *ProviderType == UE::LiveLinkHub::Private::LiveLinkHubProviderType && bAutoDiscoverLiveLinkHub)
 				{
 					AddLiveLinkSource(PollResult);
 				}
diff --git a/Engine/Plugins/Animation/LiveLinkHub/Source/LiveLinkHubMessaging/Private/LiveLinkHubMessages.cpp b/Engine/Plugins/Animation/LiveLinkHub/Source/LiveLinkHubMessaging/Private/LiveLinkHubMessages.cpp
--- a/Engine/Plugins/Animation/LiveLinkHub/Source/LiveLinkHubMessaging/Private/LiveLinkHubMessages.cpp
+++ b/Engine/Plugins/Animation/LiveLinkHub/Source/LiveLinkHubMessaging/Private/LiveLinkHubMessages.cpp
@@ -12,6 +12,7 @@
 
 const FName UE::LiveLinkHub::Private::LiveLinkHubProviderType = TEXT("LiveLinkHub");
 FName FLiveLinkHubMessageAnnotation::ProviderTypeAnnotation = TEXT("ProviderType");
+FName FLiveLinkHubMessageAnnotation::AutoDiscoverAnnotation = TEXT("AutoDiscover");
 
 
 DEFINE_LOG_CATEGORY_STATIC(LogLiveLinkHubMessages, Log, All);
diff --git a/Engine/Plugins/Animation/LiveLinkHub/Source/LiveLinkHubMessaging/Public/LiveLinkHubMessages.h b/Engine/Plugins/Animation/LiveLinkHub/Source/LiveLinkHubMessaging/Public/LiveLinkHubMessages.h
--- a/Engine/Plugins/Animation/LiveLinkHub/Source/LiveLinkHubMessaging/Public/LiveLinkHubMessages.h
+++ b/Engine/Plugins/Animation/LiveLinkHub/Source/LiveLinkHubMessaging/Public/LiveLinkHubMessages.h
@@ -11,12 +11,20 @@
 
 #include "LiveLinkHubMessages.generated.h"
 
-/** Annotation put on MessageBus messages to indicate the type of provider used. 
- * Absence of provider type means that the message comes from a regular LiveLinkProvider.
+/** 
+ * List of LiveLinkHub annotations.
  */
 struct LIVELINKHUBMESSAGING_API FLiveLinkHubMessageAnnotation
 {
+	/**
+	 * Annotation put on MessageBus messages to indicate the type of provider used.
+	 * Absence of provider type means that the message comes from a regular LiveLinkProvider.
+	 */
 	static FName ProviderTypeAnnotation;
+
+	/** Annotation to indicate if this source should be automatically added to the list of LiveLink sources. */
+	static FName AutoDiscoverAnnotation;
+
 };
 
 namespace UE::LiveLinkHub::Private
