Minor UI adjustments to LiveLinkHub

#rnx

diff --git a/Engine/Plugins/Animation/LiveLinkHub/Source/LiveLinkHub/Private/UI/Widgets/SLiveLinkHubSettings.h b/Engine/Plugins/Animation/LiveLinkHub/Source/LiveLinkHub/Private/UI/Widgets/SLiveLinkHubSettings.h
--- a/Engine/Plugins/Animation/LiveLinkHub/Source/LiveLinkHub/Private/UI/Widgets/SLiveLinkHubSettings.h
+++ b/Engine/Plugins/Animation/LiveLinkHub/Source/LiveLinkHub/Private/UI/Widgets/SLiveLinkHubSettings.h
@@ -53,25 +53,29 @@ public:
 
 		ChildSlot
 			[
-				SNew(SComboButton)
-					.ContentPadding(FMargin(0, 4.0f))
-					.ButtonStyle(&FAppStyle::Get().GetWidgetStyle<FButtonStyle>("Button"))
+				SNew(SBox)
+				.Padding(FMargin(0, 4.0f, 4.0f, 4.0f))
+				[
+					SNew(SComboButton)
+					.ContentPadding(0.f)
+					.ComboButtonStyle(&FAppStyle::Get().GetWidgetStyle<FComboButtonStyle>("ToolbarComboButton"))
 					.ForegroundColor(FSlateColor::UseStyle())
 					.ToolTipText(LOCTEXT("SettingsToolTip", "Settings"))
 					.ButtonContent()
 					[
 						SNew(SHorizontalBox)
-							+ SHorizontalBox::Slot()
-							.VAlign(VAlign_Center)
-							.Padding(4.0f, 0.0f)
-							.AutoWidth()
-							[
-								SNew(STextBlock)
-									.Font(FAppStyle::Get().GetFontStyle("FontAwesome.14"))
-									.Text(FEditorFontGlyphs::Cogs)
-							]
+						+ SHorizontalBox::Slot()
+						.VAlign(VAlign_Center)
+						.Padding(0.0f)
+						.AutoWidth()
+						[
+							SNew(STextBlock)
+							.Font(FAppStyle::Get().GetFontStyle("FontAwesome.14"))
+							.Text(FEditorFontGlyphs::Cogs)
+						]
 					]
 					.OnGetMenuContent(this, &SLiveLinkHubSettings::OnSettingsComboBoxOpen)
+				]
 			];
 	}
 
diff --git a/Engine/Plugins/Animation/LiveLinkHub/Source/LiveLinkHub/Private/UI/Widgets/SLiveLinkTimecode.cpp b/Engine/Plugins/Animation/LiveLinkHub/Source/LiveLinkHub/Private/UI/Widgets/SLiveLinkTimecode.cpp
--- a/Engine/Plugins/Animation/LiveLinkHub/Source/LiveLinkHub/Private/UI/Widgets/SLiveLinkTimecode.cpp
+++ b/Engine/Plugins/Animation/LiveLinkHub/Source/LiveLinkHub/Private/UI/Widgets/SLiveLinkTimecode.cpp
@@ -728,7 +728,7 @@ void SLiveLinkTimecode::Construct(const FArguments& InArgs)
 			+ SHorizontalBox::Slot()
 			.AutoWidth()
 			.VAlign(VAlign_Center)
-			.Padding(2, 0, 10, 0)
+			.Padding(2, 0, 5, 0)
 			[
 				SNew(STimecode)
 				.DisplayLabel(false)
@@ -741,7 +741,7 @@ void SLiveLinkTimecode::Construct(const FArguments& InArgs)
 			+ SHorizontalBox::Slot()
 			.AutoWidth()
 			.VAlign(VAlign_Center)
-			.Padding(2, 0, 10, 0)
+			.Padding(2, 0, 5, 0)
 			[
 				SNew(STextBlock)
 				.Font(FCoreStyle::Get().GetFontStyle(TEXT("NormalText")))
