Add a unique ID to the LiveLinkHub provider name if there's already an instance running.

#rnx

diff --git a/Engine/Plugins/Animation/LiveLinkHub/Source/LiveLinkHub/Private/Clients/LiveLinkHubProvider.cpp b/Engine/Plugins/Animation/LiveLinkHub/Source/LiveLinkHub/Private/Clients/LiveLinkHubProvider.cpp
--- a/Engine/Plugins/Animation/LiveLinkHub/Source/LiveLinkHub/Private/Clients/LiveLinkHubProvider.cpp
+++ b/Engine/Plugins/Animation/LiveLinkHub/Source/LiveLinkHub/Private/Clients/LiveLinkHubProvider.cpp
@@ -76,8 +76,8 @@ namespace LiveLinkHubProviderUtils
 }
 
 
-FLiveLinkHubProvider::FLiveLinkHubProvider(const TSharedRef<ILiveLinkHubSessionManager>& InSessionManager)
-	: FLiveLinkProvider(TEXT("LiveLink Hub"), false)
+FLiveLinkHubProvider::FLiveLinkHubProvider(const TSharedRef<ILiveLinkHubSessionManager>& InSessionManager, const FString& InProviderName)
+	: FLiveLinkProvider(InProviderName, false)
 	, SessionManager(InSessionManager)
 {
 	Annotations.Add(FLiveLinkHubMessageAnnotation::ProviderTypeAnnotation, UE::LiveLinkHub::Private::LiveLinkHubProviderType.ToString());
@@ -453,7 +453,7 @@ TArray<FLiveLinkHubClientId> FLiveLinkHubProvider::GetSessionClients() const
 TMap<FName, FString> FLiveLinkHubProvider::GetAnnotations() const
 {
 	TMap<FName, FString> AnnotationsCopy = Annotations;
-	
+
 	// Make a local copy of the annotations map and add the autodiscover annotation since this value might change over time.
 	AnnotationsCopy.Add(FLiveLinkHubMessageAnnotation::AutoDiscoverAnnotation, LexToString(GetDefault<ULiveLinkHubSettings>()->bAutoAddDiscoveredClients));
 	return AnnotationsCopy;
diff --git a/Engine/Plugins/Animation/LiveLinkHub/Source/LiveLinkHub/Private/Clients/LiveLinkHubProvider.h b/Engine/Plugins/Animation/LiveLinkHub/Source/LiveLinkHub/Private/Clients/LiveLinkHubProvider.h
--- a/Engine/Plugins/Animation/LiveLinkHub/Source/LiveLinkHub/Private/Clients/LiveLinkHubProvider.h
+++ b/Engine/Plugins/Animation/LiveLinkHub/Source/LiveLinkHub/Private/Clients/LiveLinkHubProvider.h
@@ -26,6 +26,7 @@ class FLiveLinkHubProvider : public FLiveLinkProvider, public ILiveLinkHubClient
 public:
 	using FLiveLinkProvider::SendClearSubjectToConnections;
 	using FLiveLinkProvider::GetLastSubjectStaticDataStruct;
+	using FLiveLinkProvider::GetProviderName;
 
 	/**
 	 * Create a message bus handler that will dispatch messages on the game thread. 
@@ -37,7 +38,7 @@ public:
 		return MakeShared<TGameThreadMessageHandler<MessageType, FLiveLinkHubProvider>>(this, Func);
 	}
 
-	FLiveLinkHubProvider(const TSharedRef<ILiveLinkHubSessionManager>& InSessionManager);
+	FLiveLinkHubProvider(const TSharedRef<ILiveLinkHubSessionManager>& InSessionManager, const FString& ProviderName);
 
 	virtual ~FLiveLinkHubProvider() override;
 
diff --git a/Engine/Plugins/Animation/LiveLinkHub/Source/LiveLinkHub/Private/LiveLinkHub.cpp b/Engine/Plugins/Animation/LiveLinkHub/Source/LiveLinkHub/Private/LiveLinkHub.cpp
--- a/Engine/Plugins/Animation/LiveLinkHub/Source/LiveLinkHub/Private/LiveLinkHub.cpp
+++ b/Engine/Plugins/Animation/LiveLinkHub/Source/LiveLinkHub/Private/LiveLinkHub.cpp
@@ -9,6 +9,7 @@
 #include "Clients/LiveLinkHubProvider.h"
 #include "Features/IModularFeatures.h"
 #include "Framework/Application/SlateApplication.h"
+#include "HAL/PlatformProcess.h"
 #include "ISettingsModule.h"
 #include "LiveLinkEditorSettings.h"
 #include "LiveLinkHubClient.h"
@@ -38,6 +39,43 @@ namespace LiveLinkHubInternalUtils
 	{
 		return FPaths::FileExists(TEXT("../../../LiveLinkHub/LiveLinkHub.uproject"));
 	}
+
+	/** Get the number of running LiveLinkHub instances. */
+	uint32 GetInstanceCount()
+	{
+		FPlatformProcess::FProcEnumerator ProcIter;
+
+		const FString ExePath = FPlatformProcess::ExecutablePath();
+		const FString ExeFileName = FPaths::GetCleanFilename(ExePath);
+
+		uint32 InstanceCount = 0;
+		while (ProcIter.MoveNext())
+		{
+			FPlatformProcess::FProcEnumInfo ProcInfo = ProcIter.GetCurrent();
+
+			if (ProcInfo.GetName() == ExeFileName)
+			{
+				InstanceCount++;
+			}
+		}
+
+		return InstanceCount;
+	}
+
+	/** Generate the name of the livelinkhub provider. */
+	FString GetProviderName()
+	{
+		FString ProviderName = TEXT("Live Link Hub");
+
+		uint32 InstanceCount = GetInstanceCount();
+		if (InstanceCount > 1)
+		{
+			// Append a unique ID if more than one instance is running.
+			ProviderName += FString::Printf(TEXT(" (%s)"), *FGuid::NewGuid().ToString().Right(4).ToLower());
+		}
+
+		return ProviderName;
+	}
 }
 
 TSharedPtr<FLiveLinkHub> FLiveLinkHub::Get()
@@ -70,7 +108,8 @@ void FLiveLinkHub::Initialize(FLiveLinkHubTicker& Ticker)
 	IModularFeatures::Get().RegisterModularFeature(ILiveLinkClient::ModularFeatureName, LiveLinkHubClient.Get());
 
 	SessionManager = MakeShared<FLiveLinkHubSessionManager>();
-	LiveLinkProvider = MakeShared<FLiveLinkHubProvider>(SessionManager.ToSharedRef());
+
+	LiveLinkProvider = MakeShared<FLiveLinkHubProvider>(SessionManager.ToSharedRef(), LiveLinkHubInternalUtils::GetProviderName());
 
 	FModuleManager::Get().LoadModule("Settings");
 	FModuleManager::Get().LoadModule("StatusBar");
diff --git a/Engine/Plugins/Animation/LiveLinkHub/Source/LiveLinkHub/Private/UI/Window/LiveLinkHubWindowController.cpp b/Engine/Plugins/Animation/LiveLinkHub/Source/LiveLinkHub/Private/UI/Window/LiveLinkHubWindowController.cpp
--- a/Engine/Plugins/Animation/LiveLinkHub/Source/LiveLinkHub/Private/UI/Window/LiveLinkHubWindowController.cpp
+++ b/Engine/Plugins/Animation/LiveLinkHub/Source/LiveLinkHub/Private/UI/Window/LiveLinkHubWindowController.cpp
@@ -2,6 +2,7 @@
 
 #include "LiveLinkHubWindowController.h"
 
+#include "Clients/LiveLinkHubProvider.h"
 #include "CoreGlobals.h"
 #include "Features/IModularFeatures.h"
 #include "Framework/Application/SlateApplication.h"
@@ -13,6 +14,7 @@
 #include "StandaloneRenderer.h"
 #include "UI/Widgets/LiveLinkHubMainTabController.h"
 #include "UI/Window/ModalWindowManager.h"
+#include "Widgets/Docking/SDockTab.h"
 
 #define LOCTEXT_NAMESPACE "LiveLinkHubWindowController"
 
@@ -68,6 +70,8 @@ void FLiveLinkHubWindowController::RestoreLayout(TSharedPtr<FAssetEditorToolkit>
 {
 	TRACE_CPUPROFILER_EVENT_SCOPE(FLiveLinkHubWindowController::RestoreLayout);
 
+	static const FTabId StandaloneTabId("StandaloneToolkit");
+
 	const FName LayoutName = TEXT("LiveLinkHub_v1.1");
 	const TSharedRef<FTabManager::FLayout> DefaultLayout =
 		FTabManager::NewLayout(LayoutName)
@@ -79,7 +83,7 @@ void FLiveLinkHubWindowController::RestoreLayout(TSharedPtr<FAssetEditorToolkit>
 			(
 				FTabManager::NewStack()
 				->SetSizeCoefficient(1.0f)
-				->AddTab("StandaloneToolkit", ETabState::ClosedTab)
+				->AddTab(StandaloneTabId, ETabState::ClosedTab)
 			)
 		);
 
@@ -98,7 +102,17 @@ void FLiveLinkHubWindowController::RestoreLayout(TSharedPtr<FAssetEditorToolkit>
 	UObject* DummyObject = GetMutableDefault<UObject>();
 	AssetEditorToolkit->InitAssetEditor(EToolkitMode::Standalone, nullptr, "LiveLinkHub", DefaultLayout, /*bCreateDefaultStandaloneMenu*/true, /*bCreateDefaultToolbar*/true, DummyObject, /*bInIsToolbarFocusable*/true, true, {});
 
-	FGlobalTabmanager::Get()->SetMainTab(FTabId("StandaloneToolkit"));
+	FGlobalTabmanager::Get()->SetMainTab(StandaloneTabId);
+
+	TSharedRef<SWidget> ProviderNameWidget =
+		SNew(SBox)
+		.VAlign(VAlign_Center)
+		[
+			SNew(STextBlock)
+				.Text(FText::FromString(FLiveLinkHub::Get()->GetLiveLinkProvider()->GetProviderName()))
+		];
+
+	FGlobalTabmanager::Get()->FindExistingLiveTab(StandaloneTabId)->SetTitleBarRightContent(ProviderNameWidget);
 }
 
 TSharedPtr<FModalWindowManager> FLiveLinkHubWindowController::InitializeSlateApplication()
