LiveLinkHub - Expand autoconnect settings to choose whether to disable autoconnect, autoconnect only on the host, or autoconnect to any computer on the network.

diff --git a/Engine/Plugins/Animation/LiveLink/Source/LiveLink/Private/LiveLinkMessageBusDiscoveryManager.cpp b/Engine/Plugins/Animation/LiveLink/Source/LiveLink/Private/LiveLinkMessageBusDiscoveryManager.cpp
--- a/Engine/Plugins/Animation/LiveLink/Source/LiveLink/Private/LiveLinkMessageBusDiscoveryManager.cpp
+++ b/Engine/Plugins/Animation/LiveLink/Source/LiveLink/Private/LiveLinkMessageBusDiscoveryManager.cpp
@@ -107,6 +107,11 @@ bool FLiveLinkMessageBusDiscoveryManager::IsRunning() const
 	return bRunning;
 }
 
+FMessageAddress FLiveLinkMessageBusDiscoveryManager::GetEndpointAddress() const
+{
+	return MessageEndpoint->GetAddress();
+}
+
 void FLiveLinkMessageBusDiscoveryManager::HandlePongMessage(const FLiveLinkPongMessage& Message, const TSharedRef<IMessageContext, ESPMode::ThreadSafe>& Context)
 {
 	FScopeLock Lock(&SourcesCriticalSection);
diff --git a/Engine/Plugins/Animation/LiveLink/Source/LiveLink/Public/LiveLinkMessageBusDiscoveryManager.h b/Engine/Plugins/Animation/LiveLink/Source/LiveLink/Public/LiveLinkMessageBusDiscoveryManager.h
--- a/Engine/Plugins/Animation/LiveLink/Source/LiveLink/Public/LiveLinkMessageBusDiscoveryManager.h
+++ b/Engine/Plugins/Animation/LiveLink/Source/LiveLink/Public/LiveLinkMessageBusDiscoveryManager.h
@@ -32,6 +32,9 @@ public:
 
 	bool IsRunning() const;
 
+	// Get the message bus address for the discovery manager's endpoint.
+	FMessageAddress GetEndpointAddress() const;
+
 private:
 	void HandlePongMessage(const FLiveLinkPongMessage& Message, const TSharedRef<IMessageContext, ESPMode::ThreadSafe>& Context);
 
diff --git a/Engine/Plugins/Animation/LiveLink/Source/LiveLinkEditor/Private/LiveLinkEditorModule.cpp b/Engine/Plugins/Animation/LiveLink/Source/LiveLinkEditor/Private/LiveLinkEditorModule.cpp
--- a/Engine/Plugins/Animation/LiveLink/Source/LiveLinkEditor/Private/LiveLinkEditorModule.cpp
+++ b/Engine/Plugins/Animation/LiveLink/Source/LiveLinkEditor/Private/LiveLinkEditorModule.cpp
@@ -110,7 +110,7 @@ public:
 		StyleSet->Set("LiveLinkHub.Playback.Icon", new IMAGE_PLUGIN_BRUSH_SVG(TEXT("Starship/Playback"), Icon16x16));
 		StyleSet->Set("LiveLinkHub.Subjects.Icon", new IMAGE_PLUGIN_BRUSH_SVG(TEXT("Starship/Subjects"), Icon16x16));
 
-		StyleSet->Set("LiveLinkHub.Autodiscover", new IMAGE_PLUGIN_BRUSH_SVG("Starship/AutoDiscover", Icon16x16));
+		StyleSet->Set("LiveLinkHub.AutoConnect", new IMAGE_PLUGIN_BRUSH_SVG("Starship/AutoConnect", Icon16x16));
 
 		StyleSet->Set("LiveLinkHub.Icon.Small", new IMAGE_PLUGIN_BRUSH(TEXT("Starship/LiveLinkHub_16x16"), Icon16x16));
 		StyleSet->Set("LiveLinkHub.Icon", new IMAGE_PLUGIN_BRUSH_SVG(TEXT("Starship/LiveLinkHub_Icon"), Icon40x40));
diff --git a/Engine/Plugins/Animation/LiveLinkHub/Source/LiveLinkHub/Private/Clients/LiveLinkHubProvider.cpp b/Engine/Plugins/Animation/LiveLinkHub/Source/LiveLinkHub/Private/Clients/LiveLinkHubProvider.cpp
--- a/Engine/Plugins/Animation/LiveLinkHub/Source/LiveLinkHub/Private/Clients/LiveLinkHubProvider.cpp
+++ b/Engine/Plugins/Animation/LiveLinkHub/Source/LiveLinkHub/Private/Clients/LiveLinkHubProvider.cpp
@@ -303,18 +303,23 @@ void FLiveLinkHubProvider::HandleHubConnectMessage(const FLiveLinkHubConnectMess
 			ClientsMap.Add(NewClient.Id, MoveTemp(NewClient));
 		}
 
-		if (GetDefault<ULiveLinkHubSettings>()->bAutoAddDiscoveredClients)
+		const bool bSameHost = Message.ClientInfo.Hostname == GetMachineName();
+		const ELiveLinkHubAutoConnectMode AutoConnectClients = GetDefault<ULiveLinkHubSettings>()->AutoConnectClients;
+		if (AutoConnectClients != ELiveLinkHubAutoConnectMode::Disabled)
 		{
-			AsyncTask(ENamedThreads::GameThread, [WeakSessionManager = SessionManager, NewClientId]()
+			if (AutoConnectClients == ELiveLinkHubAutoConnectMode::All || (AutoConnectClients == ELiveLinkHubAutoConnectMode::LocalOnly && bSameHost))
 			{
-				if (const TSharedPtr<ILiveLinkHubSessionManager> Manager = WeakSessionManager.Pin())
-				{
-					if (const TSharedPtr<ILiveLinkHubSession> CurrentSession = Manager->GetCurrentSession())
+				AsyncTask(ENamedThreads::GameThread, [WeakSessionManager = SessionManager, NewClientId]()
 					{
-						CurrentSession->AddClient(NewClientId);
-					}
-				}
-			});
+						if (const TSharedPtr<ILiveLinkHubSessionManager> Manager = WeakSessionManager.Pin())
+						{
+							if (const TSharedPtr<ILiveLinkHubSession> CurrentSession = Manager->GetCurrentSession())
+							{
+								CurrentSession->AddClient(NewClientId);
+							}
+						}
+					});
+			}
 		}
 		else
 		{
@@ -454,8 +459,9 @@ TMap<FName, FString> FLiveLinkHubProvider::GetAnnotations() const
 {
 	TMap<FName, FString> AnnotationsCopy = Annotations;
 
-	// Make a local copy of the annotations map and add the autodiscover annotation since this value might change over time.
-	AnnotationsCopy.Add(FLiveLinkHubMessageAnnotation::AutoDiscoverAnnotation, LexToString(GetDefault<ULiveLinkHubSettings>()->bAutoAddDiscoveredClients));
+	// Make a local copy of the annotations map and add the autoconnect annotation since this value might change over time.
+	const FString AutoConnectValue = StaticEnum<ELiveLinkHubAutoConnectMode>()->GetNameStringByValue(static_cast<int64>(GetDefault<ULiveLinkHubSettings>()->AutoConnectClients));
+	AnnotationsCopy.Add(FLiveLinkHubMessageAnnotation::AutoDiscoverAnnotation, AutoConnectValue);
 	return AnnotationsCopy;
 }
 
diff --git a/Engine/Plugins/Animation/LiveLinkHub/Source/LiveLinkHub/Private/Clients/SLiveLinkHubClientsView.h b/Engine/Plugins/Animation/LiveLinkHub/Source/LiveLinkHub/Private/Clients/SLiveLinkHubClientsView.h
--- a/Engine/Plugins/Animation/LiveLinkHub/Source/LiveLinkHub/Private/Clients/SLiveLinkHubClientsView.h
+++ b/Engine/Plugins/Animation/LiveLinkHub/Source/LiveLinkHub/Private/Clients/SLiveLinkHubClientsView.h
@@ -25,11 +25,11 @@
 #include "Widgets/Input/SButton.h"
 #include "Widgets/Input/SCheckBox.h"
 #include "Widgets/Layout/SScaleBox.h"
+#include "Widgets/Layout/SWidgetSwitcher.h"
 #include "Widgets/Input/SComboButton.h"
 #include "Widgets/Text/STextBlock.h"
 #include "Widgets/Views/STreeView.h"
 
-
 struct FGuid;
 class FLiveLinkHub;
 
@@ -44,6 +44,37 @@ static const FName NameColumnId = "Name";
 static const FName StatusColumnId = "Status";
 static const FName EnabledIconColumnId = "EnabledIcon";
 
+namespace LiveLinkHubClientsViewUtils
+{
+	TSharedRef<SWidget> GetAutoConnectIcon(ELiveLinkHubAutoConnectMode Mode)
+	{
+		FSlateColor IconColor;
+
+		// Note: Temporary until we have icons for autoconnect modes.
+		switch (Mode)
+		{
+			case ELiveLinkHubAutoConnectMode::Disabled:
+				IconColor = FSlateColor::UseSubduedForeground();
+				break;
+			
+			case ELiveLinkHubAutoConnectMode::All:
+				IconColor = FLinearColor::Yellow;
+				break;
+
+			case ELiveLinkHubAutoConnectMode::LocalOnly:
+				IconColor = FSlateColor::UseForeground();
+				break;
+			default:
+				checkNoEntry();
+		}
+
+		return SNew(SImage)
+		.ColorAndOpacity(IconColor)
+		.Image(FSlateIcon("LiveLinkStyle", "LiveLinkHub.AutoConnect").GetIcon());
+	}
+}
+
+
 /** Tree view item that represents either a client or a livelink subject. */
 struct FClientTreeViewItem
 {
@@ -349,16 +380,21 @@ public:
 				.AutoWidth()
 				.Padding(FMargin(4.0, 2.0))
 				[
-					SNew(SCheckBox)
-						.Padding(4.0f)
-						.ToolTipText(LOCTEXT("ToggleAutoDiscoveryToolTip", "Automatically add LiveLinkClients that were discovered."))
-						.Style(FAppStyle::Get(), "ToggleButtonCheckbox")
-						.IsChecked(this, &SLiveLinkHubClientsView::GetAutodiscoveryToggleState)
-						.OnCheckStateChanged(this, &SLiveLinkHubClientsView::ToggleAutodiscovery)
-						[
-							SNew(SImage)
-								.Image(FSlateIcon("LiveLinkStyle", "LiveLinkHub.Autodiscover").GetIcon())
-						]
+					SNew(SComboButton)
+					.ComboButtonStyle( FAppStyle::Get(), "SimpleComboButton" )
+					.ContentPadding(0.0f)
+					.ForegroundColor( FSlateColor::UseForeground() )
+					.ToolTipText_Lambda([]() 
+						{
+							const UEnum* Enum = StaticEnum<ELiveLinkHubAutoConnectMode>();
+							const int32 Index = Enum->GetIndexByValue(static_cast<int64>(GetDefault<ULiveLinkHubSettings>()->AutoConnectClients));
+							return Enum->GetToolTipTextByIndex(Index);
+						})
+					.OnGetMenuContent_Raw(this, &SLiveLinkHubClientsView::GetAutoConnectMenuContent)
+					.ButtonContent()
+					[
+						CreateAutoConnectModeSwitcher()
+					]
 				]
 				
 				+ SHorizontalBox::Slot()
@@ -828,19 +864,58 @@ private:
 		}
 	}
 
-	/** Get the autodiscover checkbox state. */
-	ECheckBoxState GetAutodiscoveryToggleState() const
+	/** Generate a dropdown to select an autoconnect modes. */
+	TSharedRef<SWidget> GetAutoConnectMenuContent()
 	{
-		return GetDefault<ULiveLinkHubSettings>()->bAutoAddDiscoveredClients ? ECheckBoxState::Checked : ECheckBoxState::Unchecked;
+		FMenuBuilder MenuBuilder(true, NULL);
+		if (TSharedPtr<FLiveLinkHub> LiveLinkHub = FLiveLinkHub::Get())
+		{
+			MenuBuilder.BeginSection(NAME_None, LOCTEXT("AutoConnectLabel", "Auto Connect Mode"));
+			UEnum* Enum = StaticEnum<ELiveLinkHubAutoConnectMode>();
+
+			for (int32 Index = 0; Index < Enum->NumEnums() - 1; Index++)
+			{
+				const ELiveLinkHubAutoConnectMode Mode = (ELiveLinkHubAutoConnectMode)Enum->GetValueByIndex(Index);
+				MenuBuilder.AddMenuEntry(Enum->GetDisplayNameTextByIndex(Index), Enum->GetToolTipTextByIndex(Index), FSlateIcon("LiveLinkStyle", "LiveLinkHub.AutoConnect"), FExecuteAction::CreateRaw(this, &SLiveLinkHubClientsView::SetAutoConnectMode, Mode));
+			}
+
+			MenuBuilder.EndSection();
+		}
+
+		return MenuBuilder.MakeWidget();
 	}
 
-	/** Enable / Disable UE client autodiscovery. */
-	void ToggleAutodiscovery(ECheckBoxState NewState)
+	/** Configure UE client AutoConnect mode. */
+	void SetAutoConnectMode(ELiveLinkHubAutoConnectMode AutoConnectMode)
 	{
-		GetMutableDefault<ULiveLinkHubSettings>()->bAutoAddDiscoveredClients = NewState == ECheckBoxState::Checked;
+		GetMutableDefault<ULiveLinkHubSettings>()->AutoConnectClients = AutoConnectMode;
 		GetMutableDefault<ULiveLinkHubSettings>()->SaveConfig();
 	}
 
+	/** Create the widget for switching autoconnect mode. */
+	TSharedRef<SWidget> CreateAutoConnectModeSwitcher()
+	{
+		TSharedRef<SWidgetSwitcher> Switcher =
+			SNew(SWidgetSwitcher)
+			.WidgetIndex_Lambda([]()
+			{
+				return StaticEnum<ELiveLinkHubAutoConnectMode>()->GetIndexByValue(static_cast<int64>(GetDefault<ULiveLinkHubSettings>()->AutoConnectClients));
+			});
+
+		const UEnum* Enum = StaticEnum<ELiveLinkHubAutoConnectMode>();
+
+		for (int32 Index = 0; Index < Enum->NumEnums() - 1; Index++)
+		{
+			const int32 Value = Enum->GetValueByIndex(Index);
+
+			Switcher->AddSlot()
+				[
+					LiveLinkHubClientsViewUtils::GetAutoConnectIcon((ELiveLinkHubAutoConnectMode)Value)
+				];
+		}
+		return Switcher;
+	}
+
 
 private:
 	/** Delegate called when a client is selected. */
diff --git a/Engine/Plugins/Animation/LiveLinkHub/Source/LiveLinkHub/Private/Settings/LiveLinkHubSettings.h b/Engine/Plugins/Animation/LiveLinkHub/Source/LiveLinkHub/Private/Settings/LiveLinkHubSettings.h
--- a/Engine/Plugins/Animation/LiveLinkHub/Source/LiveLinkHub/Private/Settings/LiveLinkHubSettings.h
+++ b/Engine/Plugins/Animation/LiveLinkHub/Source/LiveLinkHub/Private/Settings/LiveLinkHubSettings.h
@@ -3,7 +3,7 @@
 #pragma once
 
 #include "UObject/Object.h"
-
+#include "LiveLinkHubMessages.h"
 #include "Misc/FrameRate.h"
 
 #include "LiveLinkHubSettings.generated.h"
@@ -28,7 +28,7 @@ protected:
 public:
 	/** If enabled, discovered clients will be automatically added to the current session. */
 	UPROPERTY(config, EditAnywhere, Category="LiveLinkHub", DisplayName = "Automatically add discovered clients")
-	bool bAutoAddDiscoveredClients = true;
+	ELiveLinkHubAutoConnectMode AutoConnectClients = ELiveLinkHubAutoConnectMode::All;
 
 	/** The size in megabytes to buffer when streaming a recording. */
 	UPROPERTY(config, EditAnywhere, Category="LiveLinkHub", meta = (ClampMin = "1", UIMin = "1"))
diff --git a/Engine/Plugins/Animation/LiveLinkHub/Source/LiveLinkHubMessaging/LiveLinkHubMessaging.build.cs b/Engine/Plugins/Animation/LiveLinkHub/Source/LiveLinkHubMessaging/LiveLinkHubMessaging.build.cs
--- a/Engine/Plugins/Animation/LiveLinkHub/Source/LiveLinkHubMessaging/LiveLinkHubMessaging.build.cs
+++ b/Engine/Plugins/Animation/LiveLinkHub/Source/LiveLinkHubMessaging/LiveLinkHubMessaging.build.cs
@@ -21,6 +21,7 @@ namespace UnrealBuildTool.Rules
 				"Engine",
 				"LiveLink",
 				"LiveLinkInterface",
+				"Messaging",
 				"SlateCore",
 			});
 
diff --git a/Engine/Plugins/Animation/LiveLinkHub/Source/LiveLinkHubMessaging/Private/LiveLinkHubConnectionManager.h b/Engine/Plugins/Animation/LiveLinkHub/Source/LiveLinkHubMessaging/Private/LiveLinkHubConnectionManager.h
--- a/Engine/Plugins/Animation/LiveLinkHub/Source/LiveLinkHubMessaging/Private/LiveLinkHubConnectionManager.h
+++ b/Engine/Plugins/Animation/LiveLinkHub/Source/LiveLinkHubMessaging/Private/LiveLinkHubConnectionManager.h
@@ -5,6 +5,7 @@
 #include "Engine/World.h"
 #include "EngineAnalytics.h"
 #include "Features/IModularFeatures.h"
+#include "HAL/PlatformProcess.h"
 #include "ILiveLinkClient.h"
 #include "ILiveLinkModule.h"
 #include "ILiveLinkSource.h"
@@ -112,16 +113,25 @@ private:
 				const FString* ProviderType = PollResult->Annotations.Find(FLiveLinkHubMessageAnnotation::ProviderTypeAnnotation);
 		
 				// Before this annotation was added, LiveLinkHub would automatically be added, so we need to keep the previous behavior if we discovered an older LiveLinkHub instance.
-				bool bAutoDiscoverLiveLinkHub = true;
+				ELiveLinkHubAutoConnectMode AutoDiscover = ELiveLinkHubAutoConnectMode::Disabled;
 
 				if (const FString* AutoDiscoverAnnotation = PollResult->Annotations.Find(FLiveLinkHubMessageAnnotation::AutoDiscoverAnnotation))
 				{
-					LexFromString(bAutoDiscoverLiveLinkHub, **AutoDiscoverAnnotation);
+					int64 AutoDiscoverValue = StaticEnum<ELiveLinkHubAutoConnectMode>()->GetValueByName(**AutoDiscoverAnnotation);
+					if (AutoDiscoverValue != INDEX_NONE)
+					{
+						AutoDiscover = (ELiveLinkHubAutoConnectMode)AutoDiscoverValue;
+					}
 				}
 
-				if (ProviderType && *ProviderType == UE::LiveLinkHub::Private::LiveLinkHubProviderType && bAutoDiscoverLiveLinkHub)
+				const bool bSameHost = PollResult->MachineName == FPlatformProcess::ComputerName();
+
+				if (ProviderType && *ProviderType == UE::LiveLinkHub::Private::LiveLinkHubProviderType)
 				{
-					AddLiveLinkSource(PollResult);
+					if (AutoDiscover == ELiveLinkHubAutoConnectMode::All || (AutoDiscover == ELiveLinkHubAutoConnectMode::LocalOnly && bSameHost))
+					{
+						AddLiveLinkSource(PollResult);
+					}
 				}
 			}
 		}
diff --git a/Engine/Plugins/Animation/LiveLinkHub/Source/LiveLinkHubMessaging/Public/LiveLinkHubMessages.h b/Engine/Plugins/Animation/LiveLinkHub/Source/LiveLinkHubMessaging/Public/LiveLinkHubMessages.h
--- a/Engine/Plugins/Animation/LiveLinkHub/Source/LiveLinkHubMessaging/Public/LiveLinkHubMessages.h
+++ b/Engine/Plugins/Animation/LiveLinkHub/Source/LiveLinkHubMessaging/Public/LiveLinkHubMessages.h
@@ -11,6 +11,15 @@
 
 #include "LiveLinkHubMessages.generated.h"
 
+/** Whether if and how LiveLinkHub should auto-connect to UE clients on the network. */
+UENUM()
+enum class ELiveLinkHubAutoConnectMode : uint8
+{
+	Disabled   UMETA(ToolTip = "Don't add any clients automatically."),
+	All 	   UMETA(ToolTip = "Add any client that was found on the network."),
+	LocalOnly  UMETA(ToolTip = "Add any client running on this machine.")
+};
+
 /** 
  * List of LiveLinkHub annotations.
  */
diff --git a/Engine/Source/Programs/LiveLinkHub/Config/DefaultEngine.ini b/Engine/Source/Programs/LiveLinkHub/Config/DefaultEngine.ini
--- a/Engine/Source/Programs/LiveLinkHub/Config/DefaultEngine.ini
+++ b/Engine/Source/Programs/LiveLinkHub/Config/DefaultEngine.ini
@@ -55,3 +55,7 @@ ConverterDefaultPipeline=None
 [Animation.DefaultObjectSettings]
 BoneCompressionSettings="/Engine/Animation/DefaultAnimBoneCompressionSettings"
 CurveCompressionSettings="/Engine/Animation/DefaultAnimCurveCompressionSettings"
+
+[CoreRedirects]
++PropertyRedirects=(OldName="LiveLinkHubSettings.bAutoAddDiscoveredClients",NewName="AutoConnectClients")
+
diff --git a/Engine/Source/Runtime/LiveLinkMessageBusFramework/Private/LiveLinkProvider.cpp b/Engine/Source/Runtime/LiveLinkMessageBusFramework/Private/LiveLinkProvider.cpp
--- a/Engine/Source/Runtime/LiveLinkMessageBusFramework/Private/LiveLinkProvider.cpp
+++ b/Engine/Source/Runtime/LiveLinkMessageBusFramework/Private/LiveLinkProvider.cpp
@@ -310,6 +310,11 @@ void FLiveLinkProvider::SendClearSubjectToConnections(FName SubjectName)
 	MessageEndpoint->Send(FMessageEndpoint::MakeMessage<FLiveLinkClearSubject>(SubjectName), EMessageFlags::Reliable, GetAnnotations(), nullptr, MessageAddresses, FTimespan::Zero(), FDateTime::MaxValue());
 }
 
+FMessageAddress FLiveLinkProvider::GetEndpointAddress() const
+{
+	return MessageEndpoint->GetAddress();
+}
+
 bool FLiveLinkProvider::UpdateSubjectStaticData(const FName SubjectName, TSubclassOf<ULiveLinkRole> Role, FLiveLinkStaticDataStruct&& StaticData, const TMap<FName, FString>& ExtraAnnotations)
 {
 	FScopeLock Lock(&CriticalSection);
diff --git a/Engine/Source/Runtime/LiveLinkMessageBusFramework/Public/LiveLinkProviderImpl.h b/Engine/Source/Runtime/LiveLinkMessageBusFramework/Public/LiveLinkProviderImpl.h
--- a/Engine/Source/Runtime/LiveLinkMessageBusFramework/Public/LiveLinkProviderImpl.h
+++ b/Engine/Source/Runtime/LiveLinkMessageBusFramework/Public/LiveLinkProviderImpl.h
@@ -174,6 +174,9 @@ protected:
 	// Send a clear subject message to indicate that the subject should be removed from the connected client.
 	void SendClearSubjectToConnections(FName SubjectName);
 
+	// Get the MessageBus address of the provider.
+	FMessageAddress GetEndpointAddress() const;
+
 	// Constructor for derived classes that allows specifying that no endpoint should be created.
 	FLiveLinkProvider(const FString& InProviderName, bool bInCreateEndpoint);
 
