Fix compilation error in SLiveLinkHubClientsView

#rnx

diff --git a/Engine/Plugins/Animation/LiveLinkHub/Source/LiveLinkHub/Private/Clients/SLiveLinkHubClientsView.h b/Engine/Plugins/Animation/LiveLinkHub/Source/LiveLinkHub/Private/Clients/SLiveLinkHubClientsView.h
--- a/Engine/Plugins/Animation/LiveLinkHub/Source/LiveLinkHub/Private/Clients/SLiveLinkHubClientsView.h
+++ b/Engine/Plugins/Animation/LiveLinkHub/Source/LiveLinkHub/Private/Clients/SLiveLinkHubClientsView.h
@@ -46,7 +46,7 @@ static const FName EnabledIconColumnId = "EnabledIcon";
 
 namespace LiveLinkHubClientsViewUtils
 {
-	TSharedRef<SWidget> GetAutoConnectIcon(ELiveLinkHubAutoConnectMode Mode)
+	static TSharedRef<SWidget> GetAutoConnectIcon(ELiveLinkHubAutoConnectMode Mode)
 	{
 		FSlateColor IconColor;
 
