Live Link Hub: Minor recording session changes.

- Add IsRecording / OnRecordingStarted / OnRecordingStopped to ILiveLinkRecordingSessionInfo.
- Change default session name from empty string to "DefaultSession."
- Change default take number from 0 to 1.
- Factor CanRecord logic out of widget and into controller.

#rnx

diff --git a/Engine/Plugins/Animation/LiveLinkHub/Source/LiveLinkHub/Private/Recording/LiveLinkHubRecordingController.cpp b/Engine/Plugins/Animation/LiveLinkHub/Source/LiveLinkHub/Private/Recording/LiveLinkHubRecordingController.cpp
--- a/Engine/Plugins/Animation/LiveLinkHub/Source/LiveLinkHub/Private/Recording/LiveLinkHubRecordingController.cpp
+++ b/Engine/Plugins/Animation/LiveLinkHub/Source/LiveLinkHub/Private/Recording/LiveLinkHubRecordingController.cpp
@@ -2,6 +2,7 @@
 
 #include "LiveLinkHubRecordingController.h"
 
+#include "ILiveLinkRecordingSessionInfo.h"
 #include "Implementations/LiveLinkUAssetRecorder.h"
 #include "LiveLinkRecorder.h"
 #include "SLiveLinkHubRecordingView.h"
@@ -9,6 +10,7 @@
 #include "Widgets/Layout/SBox.h"
 #include "Widgets/SWidget.h"
 
+
 FLiveLinkHubRecordingController::FLiveLinkHubRecordingController()
 {
 	RecorderImplementation = MakeShared<FLiveLinkUAssetRecorder>();
@@ -17,6 +19,7 @@ FLiveLinkHubRecordingController::FLiveLinkHubRecordingController()
 TSharedRef<SWidget> FLiveLinkHubRecordingController::MakeRecordToolbarEntry()
 {
 	return SNew(SLiveLinkHubRecordingView)
+		.CanRecord_Raw(this, &FLiveLinkHubRecordingController::CanRecord)
 		.IsRecording_Raw(this, &FLiveLinkHubRecordingController::IsRecording)
 		.OnStartRecording_Raw(this, &FLiveLinkHubRecordingController::StartRecording)
 		.OnStopRecording_Raw(this, &FLiveLinkHubRecordingController::StopRecording);
@@ -24,11 +27,15 @@ TSharedRef<SWidget> FLiveLinkHubRecordingController::MakeRecordToolbarEntry()
 
 void FLiveLinkHubRecordingController::StartRecording()
 {
+	ILiveLinkRecordingSessionInfo::Get().OnRecordingStarted().Broadcast();
+
 	RecorderImplementation->StartRecording();
 }
 	
 void FLiveLinkHubRecordingController::StopRecording()
 {
+	ILiveLinkRecordingSessionInfo::Get().OnRecordingStopped().Broadcast();
+
 	RecorderImplementation->StopRecording();
 }
 
@@ -37,6 +44,19 @@ bool FLiveLinkHubRecordingController::IsRecording() const
 	return RecorderImplementation->IsRecording();
 }
 
+bool FLiveLinkHubRecordingController::CanRecord() const
+{
+	const FLiveLinkHubModule& LiveLinkHubModule = FModuleManager::Get().GetModuleChecked<FLiveLinkHubModule>("LiveLinkHub");
+	if (LiveLinkHubModule.GetPlaybackController()->IsInPlayback())
+	{
+		return false;
+	}
+
+	// TODO: Check for recording destination asset naming conflict.
+
+	return true;
+}
+
 void FLiveLinkHubRecordingController::RecordStaticData(const FLiveLinkSubjectKey& SubjectKey, TSubclassOf<ULiveLinkRole> Role, const FLiveLinkStaticDataStruct& StaticData)
 {
 	RecorderImplementation->RecordStaticData(SubjectKey, Role, StaticData);
diff --git a/Engine/Plugins/Animation/LiveLinkHub/Source/LiveLinkHub/Private/Recording/LiveLinkHubRecordingController.h b/Engine/Plugins/Animation/LiveLinkHub/Source/LiveLinkHub/Private/Recording/LiveLinkHubRecordingController.h
--- a/Engine/Plugins/Animation/LiveLinkHub/Source/LiveLinkHub/Private/Recording/LiveLinkHubRecordingController.h
+++ b/Engine/Plugins/Animation/LiveLinkHub/Source/LiveLinkHub/Private/Recording/LiveLinkHubRecordingController.h
@@ -27,6 +27,9 @@ public:
 	/** Returns whether we're currently recording. */
 	bool IsRecording() const;
 
+	/** Returns whether we can begin recording. */
+	bool CanRecord() const;
+
 	/** Record static data in the current recording. */
 	void RecordStaticData(const FLiveLinkSubjectKey& SubjectKey, TSubclassOf<class ULiveLinkRole> Role, const FLiveLinkStaticDataStruct& StaticData);
 
diff --git a/Engine/Plugins/Animation/LiveLinkHub/Source/LiveLinkHub/Private/Recording/SLiveLinkHubRecordingView.h b/Engine/Plugins/Animation/LiveLinkHub/Source/LiveLinkHub/Private/Recording/SLiveLinkHubRecordingView.h
--- a/Engine/Plugins/Animation/LiveLinkHub/Source/LiveLinkHub/Private/Recording/SLiveLinkHubRecordingView.h
+++ b/Engine/Plugins/Animation/LiveLinkHub/Source/LiveLinkHub/Private/Recording/SLiveLinkHubRecordingView.h
@@ -20,7 +20,6 @@ class FLiveLinkHub;
 
 #define LOCTEXT_NAMESPACE "LiveLinkHub.RecordingView"
 
-DECLARE_DELEGATE_RetVal(bool, FOnGetIsRecording);
 
 /**
  * Utility base class for all content that is displayed in an entire major tab.
@@ -28,8 +27,11 @@ DECLARE_DELEGATE_RetVal(bool, FOnGetIsRecording);
 class SLiveLinkHubRecordingView : public SCompoundWidget
 {
 public:
+	DECLARE_DELEGATE_RetVal(bool, FSimplePredicate);
+
 	SLATE_BEGIN_ARGS(SLiveLinkHubRecordingView) {}
-		SLATE_EVENT(FOnGetIsRecording, IsRecording)
+		SLATE_EVENT(FSimplePredicate, CanRecord)
+		SLATE_EVENT(FSimplePredicate, IsRecording)
 		SLATE_EVENT(FSimpleDelegate, OnStartRecording)
 		SLATE_EVENT(FSimpleDelegate, OnStopRecording)
 	SLATE_END_ARGS()
@@ -40,6 +42,7 @@ public:
 		*/
 	void Construct(const FArguments& InArgs)
 	{
+		CanRecordDelegate = InArgs._CanRecord;
 		IsRecordingDelegate = InArgs._IsRecording;
 		OnStartRecordingDelegate = InArgs._OnStartRecording;
 		OnStopRecordingDelegate = InArgs._OnStopRecording;
@@ -104,13 +107,14 @@ public:
 
 	bool CanRecord() const
 	{
-		const FLiveLinkHubModule& LiveLinkHubModule = FModuleManager::Get().GetModuleChecked<FLiveLinkHubModule>("LiveLinkHub");
-		return !LiveLinkHubModule.GetPlaybackController()->IsInPlayback();
+		return CanRecordDelegate.Execute();
 	}
 
 private:
-	/** Delegate used to know if the hub is current recording. */
-	FOnGetIsRecording IsRecordingDelegate;
+	/** Delegate used to know if the hub can begin recording. */
+	FSimplePredicate CanRecordDelegate;
+	/** Delegate used to know if the hub is currently recording. */
+	FSimplePredicate IsRecordingDelegate;
 	/** Delegate used to notice the hub that we should start recording. */
 	FSimpleDelegate OnStartRecordingDelegate;
 	/** Delegate used to notice the hub that we should stop recording. */
diff --git a/Engine/Plugins/Animation/LiveLinkHub/Source/LiveLinkHub/Private/Session/LiveLinkHubSessionData.h b/Engine/Plugins/Animation/LiveLinkHub/Source/LiveLinkHub/Private/Session/LiveLinkHubSessionData.h
--- a/Engine/Plugins/Animation/LiveLinkHub/Source/LiveLinkHub/Private/Session/LiveLinkHubSessionData.h
+++ b/Engine/Plugins/Animation/LiveLinkHub/Source/LiveLinkHub/Private/Session/LiveLinkHubSessionData.h
@@ -18,6 +18,12 @@ class ULiveLinkHubSessionData : public UObject
 public:
 	GENERATED_BODY()
 
+	ULiveLinkHubSessionData()
+	{
+		RecordingSessionName = TEXT("DefaultSession");
+		RecordingTakeNumber = 1;
+	}
+
 	/** Live link hub sources. */
 	UPROPERTY()
 	TArray<FLiveLinkSourcePreset> Sources;
diff --git a/Engine/Plugins/Animation/LiveLinkHub/Source/LiveLinkHub/Private/Session/LiveLinkHubSessionManager.h b/Engine/Plugins/Animation/LiveLinkHub/Source/LiveLinkHub/Private/Session/LiveLinkHubSessionManager.h
--- a/Engine/Plugins/Animation/LiveLinkHub/Source/LiveLinkHub/Private/Session/LiveLinkHubSessionManager.h
+++ b/Engine/Plugins/Animation/LiveLinkHub/Source/LiveLinkHub/Private/Session/LiveLinkHubSessionManager.h
@@ -14,6 +14,7 @@
 #include "LiveLinkHubClient.h"
 #include "LiveLinkHubLog.h"
 #include "LiveLinkSourceSettings.h"
+#include "Recording/LiveLinkHubRecordingController.h"
 #include "Settings/LiveLinkHubSettings.h"
 
 #define LOCTEXT_NAMESPACE "LiveLinkHub.SessionManager"
@@ -314,6 +315,29 @@ public:
 	{
 		return OnTakeNumberChangedDelegate;
 	}
+
+	virtual bool IsRecording() const override
+	{
+		const FLiveLinkHubModule& LiveLinkHubModule =
+			FModuleManager::Get().GetModuleChecked<FLiveLinkHubModule>("LiveLinkHub");
+
+		if (TSharedPtr<FLiveLinkHubRecordingController> RecordingController = LiveLinkHubModule.GetRecordingController())
+		{
+			return RecordingController->IsRecording();
+		}
+
+		return false;
+	}
+
+	virtual FSimpleMulticastDelegate& OnRecordingStarted() override
+	{
+		return OnRecordingStartedDelegate;
+	}
+
+	virtual FSimpleMulticastDelegate& OnRecordingStopped() override
+	{
+		return OnRecordingStoppedDelegate;
+	}
 	//~ End ILiveLinkRecordingSessionInfo interface
 
 private:
@@ -427,6 +451,12 @@ private:
 	/** Delegate triggered when the recording take number is changed. */
 	FOnSessionIntChanged OnTakeNumberChangedDelegate;
 
+	/** Delegate triggered when recording begins. */
+	FSimpleMulticastDelegate OnRecordingStartedDelegate;
+
+	/** Delegate triggered when recording ends. */
+	FSimpleMulticastDelegate OnRecordingStoppedDelegate;
+
 	/** Critical section used to synchronize access to the current session. */
 	mutable FCriticalSection CurrentSessionCS;
 };
diff --git a/Engine/Plugins/Animation/LiveLinkHub/Source/LiveLinkHub/Public/ILiveLinkRecordingSessionInfo.h b/Engine/Plugins/Animation/LiveLinkHub/Source/LiveLinkHub/Public/ILiveLinkRecordingSessionInfo.h
--- a/Engine/Plugins/Animation/LiveLinkHub/Source/LiveLinkHub/Public/ILiveLinkRecordingSessionInfo.h
+++ b/Engine/Plugins/Animation/LiveLinkHub/Source/LiveLinkHub/Public/ILiveLinkRecordingSessionInfo.h
@@ -40,4 +40,9 @@ public:
 	virtual FOnSessionStringChanged& OnSessionNameChanged() = 0;
 	virtual FOnSessionStringChanged& OnSlateNameChanged() = 0;
 	virtual FOnSessionIntChanged& OnTakeNumberChanged() = 0;
+
+	virtual bool IsRecording() const = 0;
+
+	virtual FSimpleMulticastDelegate& OnRecordingStarted() = 0;
+	virtual FSimpleMulticastDelegate& OnRecordingStopped() = 0;
 };
