Clean up Capture Manager dependencies and redirect to updated workflow in Live Link Hub

diff --git a/Engine/Plugins/Animation/LiveLinkHub/Source/LiveLinkHubEditor/Private/LiveLinkHubEditorModule.cpp b/Engine/Plugins/Animation/LiveLinkHub/Source/LiveLinkHubEditor/Private/LiveLinkHubEditorModule.cpp
--- a/Engine/Plugins/Animation/LiveLinkHub/Source/LiveLinkHubEditor/Private/LiveLinkHubEditorModule.cpp
+++ b/Engine/Plugins/Animation/LiveLinkHub/Source/LiveLinkHubEditor/Private/LiveLinkHubEditorModule.cpp
@@ -58,158 +58,7 @@ void FLiveLinkHubEditorModule::OnPostEngineInit()
 			LOCTEXT("LiveLinkHubLabel", "Live Link Hub"),
 			LOCTEXT("LiveLinkHubTooltip", "Launch the Live Link Hub app."),
 			FSlateIcon("LiveLinkStyle", "LiveLinkClient.Common.Icon.Small"),
-			FUIAction(FExecuteAction::CreateRaw(this, &FLiveLinkHubEditorModule::OpenLiveLinkHub)));
-	}
-}
-
-void FLiveLinkHubEditorModule::OpenLiveLinkHub()
-{
-	FAsyncTaskNotificationConfig NotificationConfig;
-	NotificationConfig.bKeepOpenOnFailure = true;
-	NotificationConfig.TitleText = LOCTEXT("LaunchingLiveLinkHub", "Launching Live Link Hub...");
-	NotificationConfig.LogCategory = &LogLiveLinkHubEditor;
-
-	FAsyncTaskNotification Notification(NotificationConfig);
-	const FText LaunchLiveLinkHubErrorTitle = LOCTEXT("LaunchLiveLinkHubErrorTitle", "Failed to Launch LiveLinkhub.");
-
-	// Try getting the livelinkhub app location by reading a registry key.
-	if (GetDefault<ULiveLinkHubEditorSettings>()->bDetectLiveLinkHubExecutable)
-	{
-		ILauncherPlatform* LauncherPlatform = FLauncherPlatformModule::Get();
-
-		UE::LiveLinkHubLauncherUtils::FInstalledApp LiveLinkHubApp;
-		if (UE::LiveLinkHubLauncherUtils::FindLiveLinkHubInstallation(LiveLinkHubApp))
-		{
-			// Found a LiveLinkHub installation from the launcher, so launch it that way.
-			
-			const FString LaunchLink = TEXT("apps") / LiveLinkHubApp.NamespaceId + TEXT("%3A") + LiveLinkHubApp.ItemId + TEXT("%3A") + LiveLinkHubApp.AppName + TEXT("?action=launch&silent=true");
-			FOpenLauncherOptions OpenOptions(LaunchLink);
-			if (!LauncherPlatform->OpenLauncher(OpenOptions))
-			{
-				Notification.SetComplete(
-					LaunchLiveLinkHubErrorTitle,
-					LOCTEXT("LaunchLiveLinkHubError_CouldNotOpenLauncher", "Could not launch Live Link Hub through the Epic Games Store."),
-					false
-				);
-
-			}
-			else
-			{
-				Notification.SetComplete(
-					LOCTEXT("LiveLinkHubLaunchSuccessTitle", "Launched Live Link Hub."),
-					LOCTEXT("LaunchLiveLinkHubError_LaunchSuccess", "Launching Liv Link Hub through the Epic Games Store."),
-					true
-				);
-			}
-		}
-		else
-		{
-			const EAppReturnType::Type Choice = FMessageDialog::Open(EAppMsgType::YesNo, LOCTEXT("InstallThroughEGS", "Live Link Hub is not currently installed, do you want to install it through the Epic Games Store?"));
-
-			if (Choice == EAppReturnType::Yes)
-			{
-				// Could not find LiveLinkHub from the launcher. Prompt the user to open the EGS and install it.
-				if (!GetDefault<ULiveLinkHubEditorSettings>()->LiveLinkHubStorePage.IsEmpty())
-				{
-					FOpenLauncherOptions OpenOptions(GetDefault<ULiveLinkHubEditorSettings>()->LiveLinkHubStorePage);
-					if (!LauncherPlatform->OpenLauncher(OpenOptions))
-					{
-						Notification.SetComplete(
-							LaunchLiveLinkHubErrorTitle,
-							LOCTEXT("LaunchLiveLinkHubError_CouldNotFindHubStorePage", "Could not find the Live Link Hub page on the Epic Games Store."),
-							false
-						);
-					}
-					else
-					{
-						Notification.SetComplete(
-							LaunchLiveLinkHubErrorTitle,
-							LOCTEXT("LaunchLiveLinkHub_LaunchFromStore", "Opening Epic Games Store to the Live Link Hub page."),
-							true
-						);
-					}
-
-				}
-				else
-				{
-					Notification.SetComplete(
-						LaunchLiveLinkHubErrorTitle,
-						LOCTEXT("LaunchLiveLinkHubError_EmptyConfig", "Could not find the Live Link Hub page on the Epic Games Store, missing configuration for the store page."),
-						false
-					);
-				
-				}
-			}
-			else
-			{
-				Notification.SetComplete(
-					LaunchLiveLinkHubErrorTitle,
-					LOCTEXT("LaunchLiveLinkHub_DidNotLaunchFromStore", "Live Link Hub could not be launched since it wasn't installed."),
-					false
-				);
-			}
-
-		}
-
-		return;
-	}
-
-	// Find livelink hub executable location for our build configuration
-	FString LiveLinkHubPath = FPlatformProcess::GenerateApplicationPath(TEXT("LiveLinkHub"), FApp::GetBuildConfiguration());
-
-	// Validate it exists and fall back to development if it doesn't.
-	if (!IFileManager::Get().FileExists(*LiveLinkHubPath))
-	{
-		LiveLinkHubPath = FPlatformProcess::GenerateApplicationPath(TEXT("LiveLinkHub"), EBuildConfiguration::Development);
-
-		// If it still doesn't exist, fall back to the shipping executable.
-		if (!IFileManager::Get().FileExists(*LiveLinkHubPath))
-		{
-			LiveLinkHubPath = FPlatformProcess::GenerateApplicationPath(TEXT("LiveLinkHub"), EBuildConfiguration::Shipping);
-		}
-	}
-
-	if (!IFileManager::Get().FileExists(*LiveLinkHubPath))
-	{
-		Notification.SetComplete(
-			LaunchLiveLinkHubErrorTitle,
-			LOCTEXT("LaunchLiveLinkHubError_ExecutableMissing", "Could not find the executable. Have you compiled the Live Link Hub app?"),
-			false
-		);
-
-		return;
-	}
-
-	// Validate we do not have it running locally
-	const FString AppName = FPaths::GetCleanFilename(LiveLinkHubPath);
-	if (FPlatformProcess::IsApplicationRunning(*AppName))
-	{
-		Notification.SetComplete(
-			LaunchLiveLinkHubErrorTitle,
-			LOCTEXT("LaunchLiveLinkHubError_AlreadyRunning", "A Live Link Hub instance is already running."),
-			false
-		);
-		return;
-	}
-
-	constexpr bool bLaunchDetached = true;
-	constexpr bool bLaunchHidden = false;
-	constexpr bool bLaunchReallyHidden = false;
-
-	const FProcHandle ProcHandle = FPlatformProcess::CreateProc(*LiveLinkHubPath, TEXT(""), bLaunchDetached, bLaunchHidden, bLaunchReallyHidden, nullptr, 0, nullptr, nullptr, nullptr);
-	if (ProcHandle.IsValid())
-	{
-		Notification.SetComplete(
-			LOCTEXT("LaunchedLiveLinkHub", "Launched Live Link Hub"), FText(), true);
-
-		return;
-	}
-	else // Very unlikely in practice, but possible in theory.
-	{
-		Notification.SetComplete(
-			LaunchLiveLinkHubErrorTitle,
-			LOCTEXT("LaunchLiveLinkHubError_InvalidHandle", "Failed to create the Live Link Hub process."),
-			false);
+			FUIAction(FExecuteAction::CreateStatic(&UE::LiveLinkHubLauncherUtils::OpenLiveLinkHub)));
 	}
 }
 
diff --git a/Engine/Plugins/Animation/LiveLinkHub/Source/LiveLinkHubEditor/Private/LiveLinkHubEditorModule.h b/Engine/Plugins/Animation/LiveLinkHub/Source/LiveLinkHubEditor/Private/LiveLinkHubEditorModule.h
--- a/Engine/Plugins/Animation/LiveLinkHub/Source/LiveLinkHubEditor/Private/LiveLinkHubEditorModule.h
+++ b/Engine/Plugins/Animation/LiveLinkHub/Source/LiveLinkHubEditor/Private/LiveLinkHubEditorModule.h
@@ -23,9 +23,6 @@ private:
 	/** Callback used to register the livelink hub status bar. */
 	void OnPostEngineInit();
 
-	/** Launch the livelinkhub executable. */
-	void OpenLiveLinkHub();
-
 	/** Extend the editor's bottom status bar to add the livelink hub widget. */
 	void RegisterLiveLinkHubStatusBar();
 
diff --git a/Engine/Plugins/Animation/LiveLinkHub/Source/LiveLinkHubEditor/Private/LiveLinkHubLauncherUtils.cpp b/Engine/Plugins/Animation/LiveLinkHub/Source/LiveLinkHubEditor/Private/LiveLinkHubLauncherUtils.cpp
new file mode 100644
--- /dev/null
+++ b/Engine/Plugins/Animation/LiveLinkHub/Source/LiveLinkHubEditor/Private/LiveLinkHubLauncherUtils.cpp
@@ -0,0 +1,232 @@
+// Copyright Epic Games, Inc. All Rights Reserved.
+
+#include "LiveLinkHubLauncherUtils.h"
+#include "LiveLinkHubEditorSettings.h"
+
+#include "Containers/StringConv.h"
+#include "HAL/FileManager.h"
+#include "LauncherPlatformModule.h"
+#include "Misc/App.h"
+#include "Misc/FileHelper.h"
+#include "Misc/AsyncTaskNotification.h"
+#include "Misc/MessageDialog.h"
+#include "Serialization/JsonTypes.h"
+#include "Serialization/JsonReader.h"
+#include "Serialization/JsonSerializer.h"
+#include "UObject/Object.h"
+#include "UObject/ObjectMacros.h"
+
+
+DECLARE_LOG_CATEGORY_CLASS(LogLiveLinkHubLauncher, Log, Log)
+
+#define LOCTEXT_NAMESPACE "LiveLinkHubLauncher"
+
+
+bool UE::LiveLinkHubLauncherUtils::FindLiveLinkHubInstallation(FInstalledApp& OutLiveLinkHubInfo)
+{
+	FString InstalledListFile = FString(FPlatformProcess::ApplicationSettingsDir()) / TEXT("UnrealEngineLauncher/LauncherInstalled.dat");
+
+	FString InstalledText;
+	if (FFileHelper::LoadFileToString(InstalledText, *InstalledListFile))
+	{
+		// Deserialize the object
+		TSharedPtr<FJsonObject> RootObject;
+		TSharedRef<TJsonReader<>> Reader = TJsonReaderFactory<>::Create(InstalledText);
+		if (FJsonSerializer::Deserialize(Reader, RootObject) && RootObject.IsValid())
+		{
+			// Parse the list of installations
+			TArray<TSharedPtr<FJsonValue>> InstallationList = RootObject->GetArrayField(TEXT("InstallationList"));
+			for (int32 Idx = 0; Idx < InstallationList.Num(); Idx++)
+			{
+				FInstalledApp InstalledApp;
+
+				TSharedPtr<FJsonObject> InstallationItem = InstallationList[Idx]->AsObject();
+				InstalledApp.AppName = InstallationItem->GetStringField(TEXT("AppName"));
+
+				if (InstalledApp.AppName == GetDefault<ULiveLinkHubEditorSettings>()->LiveLinkHubAppName)
+				{
+					InstalledApp.InstallLocation = InstallationItem->GetStringField(TEXT("InstallLocation"));
+
+					if (!InstalledApp.InstallLocation.Len())
+					{
+						// Shouldn't happen in theory, but just to be safe. 
+						// Doing a continue here instead of returning in case there were somehow multiple LLH installations.
+						continue;
+					}
+
+					const FString& TargetVersion = GetDefault<ULiveLinkHubEditorSettings>()->LiveLinkHubTargetVersion;
+					InstalledApp.AppVersion = InstallationItem->GetStringField(TEXT("AppVersion"));
+
+					if (!TargetVersion.IsEmpty() && InstalledApp.AppVersion != TargetVersion)
+					{
+						// If we target a specific version and it doesn't match the installed app, ignore it.
+						continue;
+					}
+
+					InstalledApp.NamespaceId = InstallationItem->GetStringField(TEXT("NamespaceId"));
+					InstalledApp.ItemId = InstallationItem->GetStringField(TEXT("ItemId"));
+					InstalledApp.ArtifactId = InstallationItem->GetStringField(TEXT("ArtifactId"));
+
+
+					OutLiveLinkHubInfo = MoveTemp(InstalledApp);
+					return true;
+				}
+			}
+		}
+	}
+
+	return false;
+}
+
+void UE::LiveLinkHubLauncherUtils::OpenLiveLinkHub()
+{
+	FAsyncTaskNotificationConfig NotificationConfig;
+	NotificationConfig.bKeepOpenOnFailure = true;
+	NotificationConfig.TitleText = LOCTEXT("LaunchingLiveLinkHub", "Launching Live Link Hub...");
+	NotificationConfig.LogCategory = &LogLiveLinkHubLauncher;
+
+	FAsyncTaskNotification Notification(NotificationConfig);
+	const FText LaunchLiveLinkHubErrorTitle = LOCTEXT("LaunchLiveLinkHubErrorTitle", "Failed to Launch LiveLinkhub.");
+
+	// Try getting the livelinkhub app location by reading a registry key.
+	if (GetDefault<ULiveLinkHubEditorSettings>()->bDetectLiveLinkHubExecutable)
+	{
+		ILauncherPlatform* LauncherPlatform = FLauncherPlatformModule::Get();
+
+		UE::LiveLinkHubLauncherUtils::FInstalledApp LiveLinkHubApp;
+		if (UE::LiveLinkHubLauncherUtils::FindLiveLinkHubInstallation(LiveLinkHubApp))
+		{
+			// Found a LiveLinkHub installation from the launcher, so launch it that way.
+
+			const FString LaunchLink = TEXT("apps") / LiveLinkHubApp.NamespaceId + TEXT("%3A") + LiveLinkHubApp.ItemId + TEXT("%3A") + LiveLinkHubApp.AppName + TEXT("?action=launch&silent=true");
+			FOpenLauncherOptions OpenOptions(LaunchLink);
+			if (!LauncherPlatform->OpenLauncher(OpenOptions))
+			{
+				Notification.SetComplete(
+					LaunchLiveLinkHubErrorTitle,
+					LOCTEXT("LaunchLiveLinkHubError_CouldNotOpenLauncher", "Could not launch Live Link Hub through the Epic Games Store."),
+					false
+				);
+
+			}
+			else
+			{
+				Notification.SetComplete(
+					LOCTEXT("LiveLinkHubLaunchSuccessTitle", "Launched Live Link Hub."),
+					LOCTEXT("LaunchLiveLinkHubError_LaunchSuccess", "Launching Liv Link Hub through the Epic Games Store."),
+					true
+				);
+			}
+		}
+		else
+		{
+			const EAppReturnType::Type Choice = FMessageDialog::Open(EAppMsgType::YesNo, LOCTEXT("InstallThroughEGS", "Live Link Hub is not currently installed, do you want to install it through the Epic Games Store?"));
+
+			if (Choice == EAppReturnType::Yes)
+			{
+				// Could not find LiveLinkHub from the launcher. Prompt the user to open the EGS and install it.
+				if (!GetDefault<ULiveLinkHubEditorSettings>()->LiveLinkHubStorePage.IsEmpty())
+				{
+					FOpenLauncherOptions OpenOptions(GetDefault<ULiveLinkHubEditorSettings>()->LiveLinkHubStorePage);
+					if (!LauncherPlatform->OpenLauncher(OpenOptions))
+					{
+						Notification.SetComplete(
+							LaunchLiveLinkHubErrorTitle,
+							LOCTEXT("LaunchLiveLinkHubError_CouldNotFindHubStorePage", "Could not find the Live Link Hub page on the Epic Games Store."),
+							false
+						);
+					}
+					else
+					{
+						Notification.SetComplete(
+							LaunchLiveLinkHubErrorTitle,
+							LOCTEXT("LaunchLiveLinkHub_LaunchFromStore", "Opening Epic Games Store to the Live Link Hub page."),
+							true
+						);
+					}
+
+				}
+				else
+				{
+					Notification.SetComplete(
+						LaunchLiveLinkHubErrorTitle,
+						LOCTEXT("LaunchLiveLinkHubError_EmptyConfig", "Could not find the Live Link Hub page on the Epic Games Store, missing configuration for the store page."),
+						false
+					);
+
+				}
+			}
+			else
+			{
+				Notification.SetComplete(
+					LaunchLiveLinkHubErrorTitle,
+					LOCTEXT("LaunchLiveLinkHub_DidNotLaunchFromStore", "Live Link Hub could not be launched since it wasn't installed."),
+					false
+				);
+			}
+
+		}
+
+		return;
+	}
+
+	// Find livelink hub executable location for our build configuration
+	FString LiveLinkHubPath = FPlatformProcess::GenerateApplicationPath(TEXT("LiveLinkHub"), FApp::GetBuildConfiguration());
+
+	// Validate it exists and fall back to development if it doesn't.
+	if (!IFileManager::Get().FileExists(*LiveLinkHubPath))
+	{
+		LiveLinkHubPath = FPlatformProcess::GenerateApplicationPath(TEXT("LiveLinkHub"), EBuildConfiguration::Development);
+
+		// If it still doesn't exist, fall back to the shipping executable.
+		if (!IFileManager::Get().FileExists(*LiveLinkHubPath))
+		{
+			LiveLinkHubPath = FPlatformProcess::GenerateApplicationPath(TEXT("LiveLinkHub"), EBuildConfiguration::Shipping);
+		}
+	}
+
+	if (!IFileManager::Get().FileExists(*LiveLinkHubPath))
+	{
+		Notification.SetComplete(
+			LaunchLiveLinkHubErrorTitle,
+			LOCTEXT("LaunchLiveLinkHubError_ExecutableMissing", "Could not find the executable. Have you compiled the Live Link Hub app?"),
+			false
+		);
+
+		return;
+	}
+
+	// Validate we do not have it running locally
+	const FString AppName = FPaths::GetCleanFilename(LiveLinkHubPath);
+	if (FPlatformProcess::IsApplicationRunning(*AppName))
+	{
+		Notification.SetComplete(
+			LaunchLiveLinkHubErrorTitle,
+			LOCTEXT("LaunchLiveLinkHubError_AlreadyRunning", "A Live Link Hub instance is already running."),
+			false
+		);
+		return;
+	}
+
+	constexpr bool bLaunchDetached = true;
+	constexpr bool bLaunchHidden = false;
+	constexpr bool bLaunchReallyHidden = false;
+
+	const FProcHandle ProcHandle = FPlatformProcess::CreateProc(*LiveLinkHubPath, TEXT(""), bLaunchDetached, bLaunchHidden, bLaunchReallyHidden, nullptr, 0, nullptr, nullptr, nullptr);
+	if (ProcHandle.IsValid())
+	{
+		Notification.SetComplete(
+			LOCTEXT("LaunchedLiveLinkHub", "Launched Live Link Hub"), FText(), true);
+
+		return;
+	}
+	else // Very unlikely in practice, but possible in theory.
+	{
+		Notification.SetComplete(
+			LaunchLiveLinkHubErrorTitle,
+			LOCTEXT("LaunchLiveLinkHubError_InvalidHandle", "Failed to create the Live Link Hub process."),
+			false);
+	}
+}
+
+#undef LOCTEXT_NAMESPACE /*LiveLinkHubLauncher*/ 
diff --git a/Engine/Plugins/Animation/LiveLinkHub/Source/LiveLinkHubEditor/Private/LiveLinkHubLauncherUtils.h b/Engine/Plugins/Animation/LiveLinkHub/Source/LiveLinkHubEditor/Private/LiveLinkHubLauncherUtils.h
deleted file mode 100644
diff --git a/Engine/Plugins/Animation/LiveLinkHub/Source/LiveLinkHubEditor/Public/LiveLinkHubLauncherUtils.h b/Engine/Plugins/Animation/LiveLinkHub/Source/LiveLinkHubEditor/Public/LiveLinkHubLauncherUtils.h
new file mode 100644
--- /dev/null
+++ b/Engine/Plugins/Animation/LiveLinkHub/Source/LiveLinkHubEditor/Public/LiveLinkHubLauncherUtils.h
@@ -0,0 +1,35 @@
+// Copyright Epic Games, Inc. All Rights Reserved.
+
+#pragma once
+
+#include "Containers/UnrealString.h"
+
+namespace UE::LiveLinkHubLauncherUtils
+{
+	struct FInstalledApp
+	{
+		/** Location of the installed app. */
+		FString InstallLocation;
+
+		/** Namespace of the app. */
+		FString NamespaceId;
+
+		/** Id of the app. */
+		FString ItemId;
+
+		/** Unique ID for the app on the EGS. */
+		FString ArtifactId;
+
+		/** Version of the app. For LiveLinkHub this will correspond to a CL number.  */
+		FString AppVersion;
+
+		/** The apps' internal name. Usually matches the ArtifactId except if the app was using a legacy publishing workflow.  */
+		FString AppName;
+	};
+
+	/** Gather all the installed apps from the Epic Games Launcher. */
+	bool LIVELINKHUBEDITOR_API FindLiveLinkHubInstallation(FInstalledApp& OutLiveLinkHubInfo);
+
+	/** Launch the livelinkhub executable. */
+	void LIVELINKHUBEDITOR_API OpenLiveLinkHub();
+}
diff --git a/Engine/Source/Programs/LiveLinkHub/Source/LiveLinkHub.Target.cs b/Engine/Source/Programs/LiveLinkHub/Source/LiveLinkHub.Target.cs
--- a/Engine/Source/Programs/LiveLinkHub/Source/LiveLinkHub.Target.cs
+++ b/Engine/Source/Programs/LiveLinkHub/Source/LiveLinkHub.Target.cs
@@ -16,7 +16,7 @@ public class LiveLinkHubTarget : TargetRules
 
 	// Whether to enable building Capture Manager plugin.
 	[CommandLine("-EnableCaptureManagerPlugin=")]
-	public bool bEnableCaptureManagerPlugin = false;
+	public bool bEnableCaptureManagerPlugin = true;
 
 	// Whether to enable building MetaHuman LiveLink plugin.
 	[CommandLine("-EnableMetaHumanLiveLinkPlugin=")]
@@ -65,7 +65,7 @@ public class LiveLinkHubTarget : TargetRules
 
 		if (bEnableCaptureManagerPlugin)
 		{
-			EnablePlugins.AddRange(new string[]
+			OptionalPlugins.AddRange(new string[]
 			{
 				"CaptureManagerApp",
 				"CaptureManagerCore"
