Live Link Hub: Add support for extending session data save/load.

Derive from ULiveLinkHubSessionExtraData to define your UProperties, then implement ILiveLinkHubSessionExtraDataHandler to respond to session save/load events.

diff --git a/Engine/Plugins/Animation/LiveLinkHub/Source/LiveLinkHub/Private/Session/LiveLinkHubSession.h b/Engine/Plugins/Animation/LiveLinkHub/Source/LiveLinkHub/Private/Session/LiveLinkHubSession.h
--- a/Engine/Plugins/Animation/LiveLinkHub/Source/LiveLinkHub/Private/Session/LiveLinkHubSession.h
+++ b/Engine/Plugins/Animation/LiveLinkHub/Source/LiveLinkHub/Private/Session/LiveLinkHubSession.h
@@ -164,25 +164,25 @@ public:
 		SessionData->TimecodeSettings = TimecodeSettings;
 	}
 
-	virtual bool ShouldUseLiveLinkHubAsCustomTimeStepSource() const
+	virtual bool ShouldUseLiveLinkHubAsCustomTimeStepSource() const override
 	{
 		FReadScopeLock Locker(SessionDataLock);
 		return SessionData->bUseLiveLinkHubAsCustomTimeStepSource;
 	}
 
-	virtual void SetShouldUseLiveLinkHubAsCustomTimeStepSource(bool bUseLiveLinkHubAsCustomTimeStepSource)
+	virtual void SetShouldUseLiveLinkHubAsCustomTimeStepSource(bool bUseLiveLinkHubAsCustomTimeStepSource) override
 	{
 		FWriteScopeLock Locker(SessionDataLock);
 		SessionData->bUseLiveLinkHubAsCustomTimeStepSource = bUseLiveLinkHubAsCustomTimeStepSource;
 	}
 
-	virtual FLiveLinkHubCustomTimeStepSettings GetCustomTimeStepSettings() const
+	virtual FLiveLinkHubCustomTimeStepSettings GetCustomTimeStepSettings() const override
 	{
 		FReadScopeLock Locker(SessionDataLock);
 		return SessionData->CustomTimeStepSettings;
 	}
 
-	virtual void SetCustomTimeStepSettings(const FLiveLinkHubCustomTimeStepSettings& CustomTimeStepSettings) const
+	virtual void SetCustomTimeStepSettings(const FLiveLinkHubCustomTimeStepSettings& CustomTimeStepSettings) const override
 	{
 		FWriteScopeLock Locker(SessionDataLock);
 		SessionData->CustomTimeStepSettings = CustomTimeStepSettings;
diff --git a/Engine/Plugins/Animation/LiveLinkHub/Source/LiveLinkHub/Private/Session/LiveLinkHubSessionData.h b/Engine/Plugins/Animation/LiveLinkHub/Source/LiveLinkHub/Private/Session/LiveLinkHubSessionData.h
--- a/Engine/Plugins/Animation/LiveLinkHub/Source/LiveLinkHub/Private/Session/LiveLinkHubSessionData.h
+++ b/Engine/Plugins/Animation/LiveLinkHub/Source/LiveLinkHub/Private/Session/LiveLinkHubSessionData.h
@@ -6,11 +6,12 @@
 #include "CoreTypes.h"
 #include "LiveLinkHubMessages.h"
 #include "LiveLinkPresetTypes.h"
+#include "LiveLinkHubSessionExtraData.h"
 #include "Misc/Guid.h"
 
-
 #include "LiveLinkHubSessionData.generated.h"
 
+
 /** Live link hub session data that can be saved to disk. */
 UCLASS()
 class ULiveLinkHubSessionData : public UObject
@@ -63,4 +64,32 @@ public:
 	/** Whether the hub should be used as a CustomTimeStep source for connected clients. */
 	UPROPERTY()
 	bool bUseLiveLinkHubAsCustomTimeStepSource = false;
+
+	/** Generic / extensible storage for additional fields to be serialized with the session. */
+	UPROPERTY(Instanced) // Instanced so that the contained types are serialized inline.
+	TMap<TSubclassOf<ULiveLinkHubSessionExtraData>, TObjectPtr<ULiveLinkHubSessionExtraData>> ExtraDatas;
+
+public:
+	ULiveLinkHubSessionExtraData* GetExtraData(TSubclassOf<ULiveLinkHubSessionExtraData> InExtraDataClass)
+	{
+		if (TObjectPtr<ULiveLinkHubSessionExtraData>* MaybeData = ExtraDatas.Find(InExtraDataClass))
+		{
+			check(*MaybeData);
+			return *MaybeData;
+		}
+
+		return nullptr;
+	}
+
+	ULiveLinkHubSessionExtraData* GetOrCreateExtraData(TSubclassOf<ULiveLinkHubSessionExtraData> InExtraDataClass)
+	{
+		if (ULiveLinkHubSessionExtraData* ExistingData = GetExtraData(InExtraDataClass))
+		{
+			return ExistingData;
+		}
+
+		ULiveLinkHubSessionExtraData* NewData = NewObject<ULiveLinkHubSessionExtraData>(this, *InExtraDataClass);
+		ExtraDatas.Add(InExtraDataClass, NewData);
+		return NewData;
+	}
 };
diff --git a/Engine/Plugins/Animation/LiveLinkHub/Source/LiveLinkHub/Private/Session/LiveLinkHubSessionManager.h b/Engine/Plugins/Animation/LiveLinkHub/Source/LiveLinkHub/Private/Session/LiveLinkHubSessionManager.h
--- a/Engine/Plugins/Animation/LiveLinkHub/Source/LiveLinkHub/Private/Session/LiveLinkHubSessionManager.h
+++ b/Engine/Plugins/Animation/LiveLinkHub/Source/LiveLinkHub/Private/Session/LiveLinkHubSessionManager.h
@@ -13,6 +13,7 @@
 #include "ILiveLinkRecordingSessionInfo.h"
 #include "LiveLinkHubClient.h"
 #include "LiveLinkHubLog.h"
+#include "LiveLinkHubSessionExtraData.h"
 #include "LiveLinkSourceSettings.h"
 #include "Recording/LiveLinkHubRecordingController.h"
 #include "Settings/LiveLinkHubSettings.h"
@@ -140,9 +141,11 @@ public:
 			return;
 		}
 
+		IModularFeatures& ModularFeatures = IModularFeatures::Get();
+
 		FLiveLinkHubModule& LiveLinkHubModule = FModuleManager::Get().GetModuleChecked<FLiveLinkHubModule>("LiveLinkHub");
 		TSharedPtr<FLiveLinkHubProvider> LiveLinkProvider = LiveLinkHubModule.GetLiveLinkProvider();
-		FLiveLinkHubClient* LiveLinkHubClient = static_cast<FLiveLinkHubClient*>(&IModularFeatures::Get().GetModularFeature<ILiveLinkClient>(ILiveLinkClient::ModularFeatureName));
+		FLiveLinkHubClient* LiveLinkHubClient = static_cast<FLiveLinkHubClient*>(&ModularFeatures.GetModularFeature<ILiveLinkClient>(ILiveLinkClient::ModularFeatureName));
 
 		check(LiveLinkProvider);
 		check(LiveLinkHubClient);
@@ -155,28 +158,41 @@ public:
 
 		ULiveLinkHubSessionData* LiveLinkHubSessionData = CastChecked<ULiveLinkHubSessionData>(StaticCastSharedPtr<FLiveLinkHubSession>(CurrentSessionPtr)->SessionData.Get());
 
+		// Write sources
 		LiveLinkHubSessionData->Sources.Empty();
-		LiveLinkHubSessionData->Subjects.Empty();
-		
 		TArray<FGuid> SourceGuids = LiveLinkHubClient->GetSources();
 		for (const FGuid& SourceGuid : SourceGuids)
 		{
 			LiveLinkHubSessionData->Sources.Add(LiveLinkHubClient->GetSourcePreset(SourceGuid, nullptr));
 		}
 
+		// Write subjects
+		LiveLinkHubSessionData->Subjects.Empty();
 		TArray<FLiveLinkSubjectKey> Subjects = LiveLinkHubClient->GetSubjects(true, true);
 		for (const FLiveLinkSubjectKey& Subject : Subjects)
 		{
 			LiveLinkHubSessionData->Subjects.Add(LiveLinkHubClient->GetSubjectPreset(Subject, nullptr));
 		}
 
+		// Write clients
+		// FIXME?: We don't clear this one first like the others?
 		const TMap<FLiveLinkHubClientId, FLiveLinkHubUEClientInfo>& ClientMap = LiveLinkProvider->GetClientsMap();
-
 		for (const TTuple<FLiveLinkHubClientId, FLiveLinkHubUEClientInfo>& ClientKeyVal : ClientMap)
 		{
 			LiveLinkHubSessionData->Clients.Add(ClientKeyVal.Value);
 		}
 
+		// Write extra data
+		LiveLinkHubSessionData->ExtraDatas.Empty();
+		TArray<ILiveLinkHubSessionExtraDataHandler*> ExtraDataHandlers =
+			ILiveLinkHubSessionExtraDataHandler::GetRegisteredHandlers();
+		for (ILiveLinkHubSessionExtraDataHandler* Handler : ExtraDataHandlers)
+		{
+			TSubclassOf<ULiveLinkHubSessionExtraData> ExtraDataClass = Handler->GetExtraDataClass();
+			ULiveLinkHubSessionExtraData* ExtraData = LiveLinkHubSessionData->GetOrCreateExtraData(ExtraDataClass);
+			Handler->OnExtraDataSessionSaving(ExtraData);
+		}
+
 		if (!SavePath.IsEmpty())
 		{
 			LastConfigPath = SavePath;
@@ -349,7 +365,8 @@ private:
 
 		if (ULiveLinkHubSessionData* SessionData = UE::LiveLinkHub::FileUtilities::Private::LoadConfig(LastConfigPath))
 		{
-			FLiveLinkHubClient* LiveLinkHubClient = static_cast<FLiveLinkHubClient*>(&IModularFeatures::Get().GetModularFeature<ILiveLinkClient>(ILiveLinkClient::ModularFeatureName));
+			IModularFeatures& ModularFeatures = IModularFeatures::Get();
+			FLiveLinkHubClient* LiveLinkHubClient = static_cast<FLiveLinkHubClient*>(&ModularFeatures.GetModularFeature<ILiveLinkClient>(ILiveLinkClient::ModularFeatureName));
 
 			const FLiveLinkHubModule& LiveLinkHubModule = FModuleManager::Get().GetModuleChecked<FLiveLinkHubModule>("LiveLinkHub");
 			const TSharedPtr<FLiveLinkHubProvider> LiveLinkProvider = LiveLinkHubModule.GetLiveLinkProvider();
@@ -399,6 +416,18 @@ private:
 			{
 				CurrentSessionPtr->AddRestoredClient(Client);
 			}
+
+			// Restore extra data
+			TArray<ILiveLinkHubSessionExtraDataHandler*> ExtraDataHandlers =
+				ILiveLinkHubSessionExtraDataHandler::GetRegisteredHandlers();
+			for (ILiveLinkHubSessionExtraDataHandler* Handler : ExtraDataHandlers)
+			{
+				TSubclassOf<ULiveLinkHubSessionExtraData> ExtraDataClass = Handler->GetExtraDataClass();
+				ULiveLinkHubSessionExtraData* ExtraData = SessionData->GetExtraData(ExtraDataClass);
+				// Might be null.
+				Handler->OnExtraDataSessionLoaded(ExtraData);
+			}
+
 			OnActiveSessionChangedDelegate.Broadcast(CurrentSessionPtr.ToSharedRef());
 		}
 	}
@@ -423,6 +452,14 @@ private:
 			CurrentSessionPtr = CurrentSession = MakeShared<FLiveLinkHubSession>(OnClientAddedToSessionDelegate, OnClientRemovedFromSessionDelegate);
 		}
 
+		// Issue nullptr extra data load events as a convenience; handlers probably still want to be aware of the session load event.
+		TArray<ILiveLinkHubSessionExtraDataHandler*> ExtraDataHandlers =
+			ILiveLinkHubSessionExtraDataHandler::GetRegisteredHandlers();
+		for (ILiveLinkHubSessionExtraDataHandler* Handler : ExtraDataHandlers)
+		{
+			Handler->OnExtraDataSessionLoaded(nullptr);
+		}
+
 		OnActiveSessionChangedDelegate.Broadcast(CurrentSessionPtr.ToSharedRef());
 	}
 
diff --git a/Engine/Plugins/Animation/LiveLinkHub/Source/LiveLinkHub/Public/LiveLinkHubSessionExtraData.h b/Engine/Plugins/Animation/LiveLinkHub/Source/LiveLinkHub/Public/LiveLinkHubSessionExtraData.h
new file mode 100644
--- /dev/null
+++ b/Engine/Plugins/Animation/LiveLinkHub/Source/LiveLinkHub/Public/LiveLinkHubSessionExtraData.h
@@ -0,0 +1,57 @@
+// Copyright Epic Games, Inc. All Rights Reserved.
+
+#pragma once
+
+#include "Features/IModularFeatures.h"
+#include "Templates/SubclassOf.h"
+#include "LiveLinkHubSessionExtraData.generated.h"
+
+
+/** Derive from this class to define additional UProperties that participate in session serialization. */
+UCLASS(Abstract)
+class LIVELINKHUB_API ULiveLinkHubSessionExtraData : public UObject
+{
+	GENERATED_BODY()
+};
+
+
+/** Implement this interface to provide session save/load handlers for your extra data subclass. */
+class ILiveLinkHubSessionExtraDataHandler : public IModularFeature
+{
+public:
+	/** Which derived type this handler will receive save/load events for. */
+	virtual TSubclassOf<ULiveLinkHubSessionExtraData> GetExtraDataClass() const = 0;
+
+	/** Update your extra data fields prior to session save. */
+	virtual void OnExtraDataSessionSaving(ULiveLinkHubSessionExtraData* InExtraData) = 0;
+
+	/** Handle session load. If the session was previously saved without your handler active, InExtraData may be null. */
+	virtual void OnExtraDataSessionLoaded(const ULiveLinkHubSessionExtraData* InExtraData) = 0;
+
+public:
+	/** The modular feature name used for handler registration/iteration. */
+	static FName GetModularFeatureName()
+	{
+		static const FName ModularFeatureName("LiveLinkHubSessionExtraDataHandler");
+		return ModularFeatureName;
+	}
+
+	/** Call this from your derived class to participate in session save/load. */
+	void RegisterExtraDataHandler()
+	{
+		IModularFeatures::Get().RegisterModularFeature(GetModularFeatureName(), this);
+	}
+
+	/** Call this from your derived class when finished. */
+	void UnregisterExtraDataHandler()
+	{
+		IModularFeatures::Get().UnregisterModularFeature(GetModularFeatureName(), this);
+	}
+
+	/** Used internally to iterate over registered implementations. */
+	static TArray<ILiveLinkHubSessionExtraDataHandler*> GetRegisteredHandlers()
+	{
+		return IModularFeatures::Get().GetModularFeatureImplementations<ILiveLinkHubSessionExtraDataHandler>(
+			ILiveLinkHubSessionExtraDataHandler::GetModularFeatureName());
+	}
+};
