LiveLinkHub: Fix being able to create new blueprint classes of the live link hub tokens. The ability to change the token class is being removed for now.

diff --git a/Engine/Plugins/Animation/LiveLinkHub/Source/LiveLinkHub/Private/Config/LiveLinkHubTemplateTokens.h b/Engine/Plugins/Animation/LiveLinkHub/Source/LiveLinkHub/Private/Config/LiveLinkHubTemplateTokens.h
--- a/Engine/Plugins/Animation/LiveLinkHub/Source/LiveLinkHub/Private/Config/LiveLinkHubTemplateTokens.h
+++ b/Engine/Plugins/Animation/LiveLinkHub/Source/LiveLinkHub/Private/Config/LiveLinkHubTemplateTokens.h
@@ -7,8 +7,8 @@
 
 #include "LiveLinkHubTemplateTokens.generated.h"
 
-UCLASS()
-class ULiveLinkHubNamingTokenContext : public UNamingTokenContext
+UCLASS(NotBlueprintable)
+class ULiveLinkHubNamingTokenContext final: public UNamingTokenContext
 {
 	GENERATED_BODY()
 
@@ -18,8 +18,8 @@ public:
 	FString SessionName;
 };
 
-UCLASS()
-class ULiveLinkHubNamingTokens : public UNamingTokens
+UCLASS(NotBlueprintable)
+class ULiveLinkHubNamingTokens final: public UNamingTokens
 {
 	GENERATED_BODY()
 
diff --git a/Engine/Plugins/Animation/LiveLinkHub/Source/LiveLinkHub/Private/Recording/Implementations/LiveLinkUAssetRecorder.cpp b/Engine/Plugins/Animation/LiveLinkHub/Source/LiveLinkHub/Private/Recording/Implementations/LiveLinkUAssetRecorder.cpp
--- a/Engine/Plugins/Animation/LiveLinkHub/Source/LiveLinkHub/Private/Recording/Implementations/LiveLinkUAssetRecorder.cpp
+++ b/Engine/Plugins/Animation/LiveLinkHub/Source/LiveLinkHub/Private/Recording/Implementations/LiveLinkUAssetRecorder.cpp
@@ -157,7 +157,7 @@ bool FLiveLinkUAssetRecorder::GetSavePresetPackageName(FString& OutName)
 	TMap<FString, FStringFormatArg> FormatArgs;
 	FormatArgs.Add(TEXT("date"), Today.ToString());
 
-	const TObjectPtr<ULiveLinkHubNamingTokens> Tokens = GetDefault<ULiveLinkHubSettings>()->NamingTokens;
+	const TObjectPtr<ULiveLinkHubNamingTokens> Tokens = GetDefault<ULiveLinkHubSettings>()->GetNamingTokens();
 	FNamingTokenResultData TemplateData = GEditor->GetEditorSubsystem<UNamingTokensEditorSubsystem>()->EvaluateTokenString(GetDefault<ULiveLinkHubSettings>()->FilenameTemplate, Tokens);
 
 	const FString DefaultName = TemplateData.FileName;
diff --git a/Engine/Plugins/Animation/LiveLinkHub/Source/LiveLinkHub/Private/Settings/LiveLinkHubSettings.cpp b/Engine/Plugins/Animation/LiveLinkHub/Source/LiveLinkHub/Private/Settings/LiveLinkHubSettings.cpp
--- a/Engine/Plugins/Animation/LiveLinkHub/Source/LiveLinkHub/Private/Settings/LiveLinkHubSettings.cpp
+++ b/Engine/Plugins/Animation/LiveLinkHub/Source/LiveLinkHub/Private/Settings/LiveLinkHubSettings.cpp
@@ -2,7 +2,6 @@
 
 #include "LiveLinkHubSettings.h"
 
-#include "Config/LiveLinkHubFileUtilities.h"
 #include "Config/LiveLinkHubTemplateTokens.h"
 #include "NamingTokensEditorSubsystem.h"
 
@@ -16,40 +15,20 @@ void ULiveLinkHubSettings::PostEditChangeProperty(FPropertyChangedEvent& Propert
 	{
 		CalculateExampleOutput();
 	}
-	else if (PropertyChangedEvent.GetPropertyName() == GET_MEMBER_NAME_CHECKED(ULiveLinkHubSettings, NamingTokensClass))
-	{
-		CreateNamingTokens();
-		CalculateExampleOutput();
-	}
-}
-
-void ULiveLinkHubSettings::CreateNamingTokens()
-{
-	if (NamingTokensClass)
-	{
-		NamingTokens = NewObject<ULiveLinkHubNamingTokens>(GetTransientPackage(), NamingTokensClass.Get());
-		NamingTokens->CreateDefaultTokens();
-	}
-	else
-	{
-		NamingTokens = nullptr;
-	}
 }
 
 ULiveLinkHubSettings::ULiveLinkHubSettings()
 {
-	NamingTokensClass = ULiveLinkHubNamingTokens::StaticClass();
 }
 
 void ULiveLinkHubSettings::PostInitProperties()
 {
 	Super::PostInitProperties();
-	CreateNamingTokens();
 }
 
 void ULiveLinkHubSettings::CalculateExampleOutput()
 {
-	if (const TObjectPtr<ULiveLinkHubNamingTokens> Tokens = GetDefault<ULiveLinkHubSettings>()->NamingTokens)
+	if (const TObjectPtr<ULiveLinkHubNamingTokens> Tokens = GetNamingTokens())
 	{
 		const FNamingTokenResultData TemplateData = GEditor->GetEditorSubsystem<UNamingTokensEditorSubsystem>()->EvaluateTokenString(
 			GetDefault<ULiveLinkHubSettings>()->FilenameTemplate,Tokens);
@@ -60,3 +39,15 @@ void ULiveLinkHubSettings::CalculateExampleOutput()
 		FilenameOutput = GetDefault<ULiveLinkHubSettings>()->FilenameTemplate;
 	}
 }
+
+TObjectPtr<ULiveLinkHubNamingTokens> ULiveLinkHubSettings::GetNamingTokens() const
+{
+	if (NamingTokens == nullptr)
+	{
+		NamingTokens = NewObject<ULiveLinkHubNamingTokens>(const_cast<ULiveLinkHubSettings*>(this),
+			ULiveLinkHubNamingTokens::StaticClass());
+		NamingTokens->CreateDefaultTokens();
+	}
+
+	return NamingTokens;
+}
diff --git a/Engine/Plugins/Animation/LiveLinkHub/Source/LiveLinkHub/Private/Settings/LiveLinkHubSettings.h b/Engine/Plugins/Animation/LiveLinkHub/Source/LiveLinkHub/Private/Settings/LiveLinkHubSettings.h
--- a/Engine/Plugins/Animation/LiveLinkHub/Source/LiveLinkHub/Private/Settings/LiveLinkHubSettings.h
+++ b/Engine/Plugins/Animation/LiveLinkHub/Source/LiveLinkHub/Private/Settings/LiveLinkHubSettings.h
@@ -25,14 +25,14 @@ public:
 	
 	/** Parse templates and set example output fields. */
 	void CalculateExampleOutput();
+
+	/** Get the naming tokens for Live Link Hub. */
+	TObjectPtr<ULiveLinkHubNamingTokens> GetNamingTokens() const;
 	
 protected:
 	//~ Begin UObject interface
 	virtual void PostEditChangeProperty(FPropertyChangedEvent& PropertyChangedEvent) override;
 	//~ End UObject interface
-
-	/** Create or destroy naming tokens based on the naming tokens class. */
-	void CreateNamingTokens();
 	
 public:
 	/** If enabled, discovered clients will be automatically added to the current session. */
@@ -78,19 +78,16 @@ public:
 	/** Example parsed output of the template. */
 	UPROPERTY(VisibleAnywhere, Category="Templates", DisplayName="Output")
 	FString FilenameOutput;
+	
+	/** Placeholder for a list of the automatic tokens, set from the customization. */
+	UPROPERTY(VisibleAnywhere, Category="Templates")
+	FText AutomaticTokens;
 
-	/** The class to use for the naming tokens. */
-	UPROPERTY(config, EditAnywhere, NoClear, Category="Templates")
-	TSubclassOf<ULiveLinkHubNamingTokens> NamingTokensClass;
-
+private:
 	/**
 	 * Naming tokens for Live Link, instantiated each load based on the naming tokens class.
 	 * This isn't serialized to the config file, and exists here for singleton-like access.
 	 */
-	UPROPERTY(Instanced)
-	TObjectPtr<ULiveLinkHubNamingTokens> NamingTokens;
-	
-	/** Placeholder for a list of the automatic tokens, set from the customization. */
-	UPROPERTY(VisibleAnywhere, Category="Templates")
-	FText AutomaticTokens;
+	UPROPERTY(Instanced, Transient)
+	mutable TObjectPtr<ULiveLinkHubNamingTokens> NamingTokens;
 };
diff --git a/Engine/Plugins/Animation/LiveLinkHub/Source/LiveLinkHub/Private/Settings/LiveLinkHubSettingsCustomization.cpp b/Engine/Plugins/Animation/LiveLinkHub/Source/LiveLinkHub/Private/Settings/LiveLinkHubSettingsCustomization.cpp
--- a/Engine/Plugins/Animation/LiveLinkHub/Source/LiveLinkHub/Private/Settings/LiveLinkHubSettingsCustomization.cpp
+++ b/Engine/Plugins/Animation/LiveLinkHub/Source/LiveLinkHub/Private/Settings/LiveLinkHubSettingsCustomization.cpp
@@ -57,7 +57,7 @@ FText FLiveLinkHubSettingsCustomization::GetDisplayTokenText() const
 	if (GEditor)
 	{
 		FNamingTokenFilterArgs Args;
-		if (const TObjectPtr<ULiveLinkHubNamingTokens> Tokens = GetDefault<ULiveLinkHubSettings>()->NamingTokens)
+		if (const TObjectPtr<ULiveLinkHubNamingTokens> Tokens = GetDefault<ULiveLinkHubSettings>()->GetNamingTokens())
 		{
 			Args.NamespacesToInclude.Add(Tokens->GetNamespace());
 		}
