LiveLinkHub: An initial load with a mismatched frame index will now fail gracefully, rather than potentially freezing or crashing.

diff --git a/Engine/Plugins/Animation/LiveLinkHub/Source/LiveLinkHub/Private/Recording/Implementations/LiveLinkUAssetRecording.cpp b/Engine/Plugins/Animation/LiveLinkHub/Source/LiveLinkHub/Private/Recording/Implementations/LiveLinkUAssetRecording.cpp
--- a/Engine/Plugins/Animation/LiveLinkHub/Source/LiveLinkHub/Private/Recording/Implementations/LiveLinkUAssetRecording.cpp
+++ b/Engine/Plugins/Animation/LiveLinkHub/Source/LiveLinkHub/Private/Recording/Implementations/LiveLinkUAssetRecording.cpp
@@ -220,6 +220,12 @@ bool ULiveLinkUAssetRecording::WaitForBufferedFrames(int32 InMinFrame, int32 InM
 		// Max frames isn't set until after the initial load.
 		while (!bPerformedInitialLoad)
 		{
+			if (bCancelStream)
+			{
+				// Likely encountered an error.
+				UnloadRecordingData();
+				return false;
+			}
 			FPlatformProcess::Sleep(0.002);
 		}
 
@@ -496,7 +502,7 @@ void ULiveLinkUAssetRecording::LoadRecordingAsync(int32 InStartFrame, int32 InCu
 	{
 		// Always set to true. Some blocking operations wait for this, and in the case of a non-fatal error
 		// we want to display error logs and don't want the program to freeze.
-		bPerformedInitialLoad = true;
+		bPerformedInitialLoad = !bCancelStream;
 	};
 	
 	if (IsFrameRangeBuffered(RangeHelpers::Private::MakeInclusiveRange(InStartFrame, FMath::Min(EndFrame, MaxPossibleFrame > 0 ? MaxPossibleFrame : EndFrame))))
@@ -535,6 +541,7 @@ void ULiveLinkUAssetRecording::LoadRecordingAsync(int32 InStartFrame, int32 InCu
 			FrameData::Private::FFrameMetaData TemporaryFrameData;
 			if (!LoadInitialFrameData(TemporaryFrameData))
 			{
+				bCancelStream = true;
 				return;
 			}
 				
@@ -555,6 +562,7 @@ void ULiveLinkUAssetRecording::LoadRecordingAsync(int32 InStartFrame, int32 InCu
 			FrameData::Private::FFrameMetaData KeyPosition;
 			if (!LoadInitialFrameData(KeyPosition))
 			{
+				bCancelStream = true;
 				return;
 			}
 
@@ -722,7 +730,10 @@ bool ULiveLinkUAssetRecording::LoadInitialFrameData(UE::LiveLinkHub::FrameData::
 	{
 		const int32 LastFrame = OutFrameData.MaxFrames - 1;
 		double Timestamp;
-		LoadTimestampFromDisk(LastFrame, OutFrameData, Timestamp);
+		if (!LoadTimestampFromDisk(LastFrame, OutFrameData, Timestamp))
+		{
+			return false;
+		}
 		OutFrameData.LastTimestamp = Timestamp;
 
 		if (Timestamp > RecordingLastTimestamp)
@@ -1196,7 +1207,7 @@ bool ULiveLinkUAssetRecording::LoadTimestampFromDisk(const int32 InFrame, const
 	Reader->GetMemoryReader() << ParsedFrameIdx;
 
 	// Ensure the parsed frame index matches the expected frame
-	if (!ensure(ParsedFrameIdx == InFrame))
+	if (ParsedFrameIdx != InFrame)
 	{
 		UE_LOG(LogLiveLinkHub, Error, TEXT("Frame index mismatch: expected %d, got %d"), InFrame, ParsedFrameIdx);
 		return false;
diff --git a/Engine/Plugins/Animation/LiveLinkHub/Source/LiveLinkHub/Private/Recording/Implementations/LiveLinkUAssetRecordingPlayer.cpp b/Engine/Plugins/Animation/LiveLinkHub/Source/LiveLinkHub/Private/Recording/Implementations/LiveLinkUAssetRecordingPlayer.cpp
--- a/Engine/Plugins/Animation/LiveLinkHub/Source/LiveLinkHub/Private/Recording/Implementations/LiveLinkUAssetRecordingPlayer.cpp
+++ b/Engine/Plugins/Animation/LiveLinkHub/Source/LiveLinkHub/Private/Recording/Implementations/LiveLinkUAssetRecordingPlayer.cpp
@@ -333,7 +333,7 @@ FFrameRate FLiveLinkPlaybackTracks::GetInitialFrameRate() const
 	return FFrameRate(30, 1);
 }
 
-void FLiveLinkUAssetRecordingPlayer::PreparePlayback(ULiveLinkRecording* CurrentRecording)
+bool FLiveLinkUAssetRecordingPlayer::PreparePlayback(ULiveLinkRecording* CurrentRecording)
 {
 	// Ensure there nothing is playing and all settings are default. It's possible the CurrentRecording has settings that need to be cleared,
 	// such as if this was just recorded and is now being loaded.
@@ -344,7 +344,7 @@ void FLiveLinkUAssetRecordingPlayer::PreparePlayback(ULiveLinkRecording* Current
 	
 	CurrentRecordingPlayback = FLiveLinkPlaybackTracks();
 
-	StreamPlayback(0);
+	return StreamPlayback(0);
 }
 
 void FLiveLinkUAssetRecordingPlayer::ShutdownPlayback()
diff --git a/Engine/Plugins/Animation/LiveLinkHub/Source/LiveLinkHub/Private/Recording/Implementations/LiveLinkUAssetRecordingPlayer.h b/Engine/Plugins/Animation/LiveLinkHub/Source/LiveLinkHub/Private/Recording/Implementations/LiveLinkUAssetRecordingPlayer.h
--- a/Engine/Plugins/Animation/LiveLinkHub/Source/LiveLinkHub/Private/Recording/Implementations/LiveLinkUAssetRecordingPlayer.h
+++ b/Engine/Plugins/Animation/LiveLinkHub/Source/LiveLinkHub/Private/Recording/Implementations/LiveLinkUAssetRecordingPlayer.h
@@ -90,7 +90,7 @@ public:
 class FLiveLinkUAssetRecordingPlayer : public ILiveLinkRecordingPlayer
 {
 public:
-	virtual void PreparePlayback(class ULiveLinkRecording* CurrentRecording) override;
+	virtual bool PreparePlayback(class ULiveLinkRecording* CurrentRecording) override;
 
 	virtual void ShutdownPlayback() override;
 	
diff --git a/Engine/Plugins/Animation/LiveLinkHub/Source/LiveLinkHub/Private/Recording/LiveLinkHubPlaybackController.cpp b/Engine/Plugins/Animation/LiveLinkHub/Source/LiveLinkHub/Private/Recording/LiveLinkHubPlaybackController.cpp
--- a/Engine/Plugins/Animation/LiveLinkHub/Source/LiveLinkHub/Private/Recording/LiveLinkHubPlaybackController.cpp
+++ b/Engine/Plugins/Animation/LiveLinkHub/Source/LiveLinkHub/Private/Recording/LiveLinkHubPlaybackController.cpp
@@ -181,7 +181,13 @@ void FLiveLinkHubPlaybackController::PreparePlayback(ULiveLinkRecording* InLiveL
 				TGuardValue<bool> PreparingPlaybackGuard(bIsPreparingPlayback, true);
 
 				RecordingToPlay.Reset(RecordingStrongPtr.Get());
-				RecordingPlayer->PreparePlayback(RecordingToPlay.Get());
+				if (!RecordingPlayer->PreparePlayback(RecordingToPlay.Get()))
+				{
+					// Something failed during playback, cancel out.
+					const FString FileName = RecordingToPlay.IsValid() ? RecordingToPlay.Get()->GetName() : "Unknown";
+					UE_LOG(LogLiveLinkHub, Error, TEXT("Playback failed for '%s', the file may be corrupted or unsupported."), *FileName);
+					return;
+				}
 
 				// The start and end of playback.
 				SetSelectionStartTime(FQualifiedFrameTime(FFrameTime::FromDecimal(0.f), GetFrameRate()));
diff --git a/Engine/Plugins/Animation/LiveLinkHub/Source/LiveLinkHub/Private/Recording/LiveLinkRecordingPlayer.h b/Engine/Plugins/Animation/LiveLinkHub/Source/LiveLinkHub/Private/Recording/LiveLinkRecordingPlayer.h
--- a/Engine/Plugins/Animation/LiveLinkHub/Source/LiveLinkHub/Private/Recording/LiveLinkRecordingPlayer.h
+++ b/Engine/Plugins/Animation/LiveLinkHub/Source/LiveLinkHub/Private/Recording/LiveLinkRecordingPlayer.h
@@ -30,7 +30,7 @@ public:
 	virtual ~ILiveLinkRecordingPlayer() = default;
 
 	/** Initialize internal structures needed for playback of the recorded data. */
-	virtual void PreparePlayback(ULiveLinkRecording* Recording) = 0;
+	virtual bool PreparePlayback(ULiveLinkRecording* Recording) = 0;
 
 	/** Preform cleanup when exiting playback. */
 	virtual void ShutdownPlayback() = 0;
