LiveLinkHub: Fix being able to play an empty recording. This prevents a possible freeze when looping and changing direction.

diff --git a/Engine/Plugins/Animation/LiveLinkHub/Source/LiveLinkHub/Private/Recording/LiveLinkHubPlaybackController.cpp b/Engine/Plugins/Animation/LiveLinkHub/Source/LiveLinkHub/Private/Recording/LiveLinkHubPlaybackController.cpp
--- a/Engine/Plugins/Animation/LiveLinkHub/Source/LiveLinkHub/Private/Recording/LiveLinkHubPlaybackController.cpp
+++ b/Engine/Plugins/Animation/LiveLinkHub/Source/LiveLinkHub/Private/Recording/LiveLinkHubPlaybackController.cpp
@@ -184,14 +184,21 @@ void FLiveLinkHubPlaybackController::PreparePlayback(ULiveLinkRecording* InLiveL
 				if (!RecordingPlayer->PreparePlayback(RecordingToPlay.Get()))
 				{
 					// Something failed during playback, cancel out.
-					const FString FileName = RecordingToPlay.IsValid() ? RecordingToPlay.Get()->GetName() : "Unknown";
-					UE_LOG(LogLiveLinkHub, Error, TEXT("Playback failed for '%s', the file may be corrupted or unsupported."), *FileName);
+					UE_LOG(LogLiveLinkHub, Error, TEXT("Playback failed for '%s', the file may be corrupted or unsupported."), *GetRecordingName());
 					return;
 				}
 
+				const FQualifiedFrameTime RecordingLength = GetLength();
+				if (RecordingLength.AsSeconds() <= 0.0)
+				{
+					// No point in continuing with an empty recording. Prevent bIsReady from being set, so the playback widget can't be used.
+					UE_LOG(LogLiveLinkHub, Warning, TEXT("Recording '%s' is empty."), *GetRecordingName());
+					return;
+				}
+				
 				// The start and end of playback.
 				SetSelectionStartTime(FQualifiedFrameTime(FFrameTime::FromDecimal(0.f), GetFrameRate()));
-				SetSelectionEndTime(GetLength());
+				SetSelectionEndTime(RecordingLength);
 		
 				// The range the user sees.
 				SliderViewRange = TRange<double>(SelectionStartTime.AsSeconds(), SelectionEndTime.AsSeconds());
@@ -472,6 +479,11 @@ UE::LiveLinkHub::RangeHelpers::Private::TRangeArray<int32> FLiveLinkHubPlaybackC
 	return RecordingPlayer->GetBufferedFrameRanges();
 }
 
+FString FLiveLinkHubPlaybackController::GetRecordingName() const
+{
+	return RecordingToPlay.IsValid() ? RecordingToPlay.Get()->GetName() : "Unknown";
+}
+
 void FLiveLinkHubPlaybackController::Start()
 {
 	FString ThreadName = TEXT("LiveLinkHub Playback Controller ");
diff --git a/Engine/Plugins/Animation/LiveLinkHub/Source/LiveLinkHub/Private/Recording/LiveLinkHubPlaybackController.h b/Engine/Plugins/Animation/LiveLinkHub/Source/LiveLinkHub/Private/Recording/LiveLinkHubPlaybackController.h
--- a/Engine/Plugins/Animation/LiveLinkHub/Source/LiveLinkHub/Private/Recording/LiveLinkHubPlaybackController.h
+++ b/Engine/Plugins/Animation/LiveLinkHub/Source/LiveLinkHub/Private/Recording/LiveLinkHubPlaybackController.h
@@ -148,6 +148,9 @@ public:
 	{
 		bLoopPlayback.store(bInLoop);
 	}
+
+	/** Retrieve the current recording's name. */
+	FString GetRecordingName() const;
 	
 	/** Delegate called when playback is finished (if recording is not set to loop). */
 	FTSSimpleMulticastDelegate& OnPlaybackFinished()
