Add Topology Mode support to LiveLink's messaging layer

diff --git a/Engine/Plugins/Animation/LiveLink/Source/LiveLink/Private/LiveLinkMessageBusSourceSettings.h b/Engine/Plugins/Animation/LiveLink/Source/LiveLink/Private/LiveLinkMessageBusSourceSettings.h
deleted file mode 100644
diff --git a/Engine/Plugins/Animation/LiveLink/Source/LiveLink/Public/LiveLinkMessageBusFinder.h b/Engine/Plugins/Animation/LiveLink/Source/LiveLink/Public/LiveLinkMessageBusFinder.h
--- a/Engine/Plugins/Animation/LiveLink/Source/LiveLink/Public/LiveLinkMessageBusFinder.h
+++ b/Engine/Plugins/Animation/LiveLink/Source/LiveLink/Public/LiveLinkMessageBusFinder.h
@@ -15,7 +15,7 @@ struct FLiveLinkPongMessage;
 
 namespace LiveLinkMessageBusHelper
 {
-	double CalculateProviderMachineOffset(double SourceMachinePlatformSeconds, const TSharedRef<IMessageContext, ESPMode::ThreadSafe>& Context);
+	LIVELINK_API double CalculateProviderMachineOffset(double SourceMachinePlatformSeconds, const TSharedRef<IMessageContext, ESPMode::ThreadSafe>& Context);
 }
 
 USTRUCT(BlueprintType)
@@ -52,6 +52,9 @@ public:
 		, Annotations(MoveTemp(InAnnotations))
 	{}
 
+	// Get  the annotations that were present on the message.
+	const TMap<FName, FString>& GetAnnotations() const { return Annotations; }
+
 	// Messagebus address of the provider
 	FMessageAddress Address;
 
diff --git a/Engine/Plugins/Animation/LiveLink/Source/LiveLink/Public/LiveLinkMessageBusSourceSettings.h b/Engine/Plugins/Animation/LiveLink/Source/LiveLink/Public/LiveLinkMessageBusSourceSettings.h
new file mode 100644
--- /dev/null
+++ b/Engine/Plugins/Animation/LiveLink/Source/LiveLink/Public/LiveLinkMessageBusSourceSettings.h
@@ -0,0 +1,24 @@
+// Copyright Epic Games, Inc. All Rights Reserved.
+
+#pragma once
+
+#include "LiveLinkSourceSettings.h"
+
+#include "LiveLinkMessageBusSourceSettings.generated.h"
+
+
+
+
+/**
+ * Settings for LiveLinkMessageBusSource.
+ * Used to apply default Evaluation mode from project settings when constructed
+ */
+UCLASS()
+class LIVELINK_API ULiveLinkMessageBusSourceSettings : public ULiveLinkSourceSettings
+{
+public:
+	GENERATED_BODY()
+
+public:
+	ULiveLinkMessageBusSourceSettings();
+};
diff --git a/Engine/Plugins/Animation/LiveLinkHub/Source/LiveLinkHub/Private/Clients/LiveLinkHubProvider.cpp b/Engine/Plugins/Animation/LiveLinkHub/Source/LiveLinkHub/Private/Clients/LiveLinkHubProvider.cpp
--- a/Engine/Plugins/Animation/LiveLinkHub/Source/LiveLinkHub/Private/Clients/LiveLinkHubProvider.cpp
+++ b/Engine/Plugins/Animation/LiveLinkHub/Source/LiveLinkHub/Private/Clients/LiveLinkHubProvider.cpp
@@ -461,8 +461,12 @@ TMap<FName, FString> FLiveLinkHubProvider::GetAnnotations() const
 
 	// Make a local copy of the annotations map and add the autoconnect annotation since this value might change over time.
 	const FString AutoConnectValue = StaticEnum<ELiveLinkHubAutoConnectMode>()->GetNameStringByValue(static_cast<int64>(GetDefault<ULiveLinkHubSettings>()->AutoConnectClients));
-	AnnotationsCopy.Add(FLiveLinkHubMessageAnnotation::AutoDiscoverAnnotation, AutoConnectValue);
-	return AnnotationsCopy;
+	AnnotationsCopy.Add(FLiveLinkHubMessageAnnotation::AutoConnectModeAnnotation, AutoConnectValue);
+
+	TMap<FName, FString> BaseAnnotations = FLiveLinkProvider::GetAnnotations();
+	BaseAnnotations.Append(AnnotationsCopy);
+
+	return BaseAnnotations;
 }
 
 TArray<FLiveLinkHubClientId> FLiveLinkHubProvider::GetDiscoveredClients() const
diff --git a/Engine/Plugins/Animation/LiveLinkHub/Source/LiveLinkHubMessaging/Private/LiveLinkHubConnectionManager.h b/Engine/Plugins/Animation/LiveLinkHub/Source/LiveLinkHubMessaging/Private/LiveLinkHubConnectionManager.h
--- a/Engine/Plugins/Animation/LiveLinkHub/Source/LiveLinkHubMessaging/Private/LiveLinkHubConnectionManager.h
+++ b/Engine/Plugins/Animation/LiveLinkHub/Source/LiveLinkHubMessaging/Private/LiveLinkHubConnectionManager.h
@@ -10,8 +10,8 @@
 #include "ILiveLinkModule.h"
 #include "ILiveLinkSource.h"
 #include "LiveLinkHubMessages.h"
-#include "ILiveLinkHubMessagingModule.h"
 #include "LiveLinkHubMessageBusSource.h"
+#include "LiveLinkHubMessagingModule.h"
 #include "LiveLinkMessageBusDiscoveryManager.h"
 #include "LiveLinkMessageBusFinder.h"
 #include "LiveLinkSettings.h"
@@ -27,10 +27,11 @@
 
 DEFINE_LOG_CATEGORY_STATIC(LogLiveLinkHubConnectionManager, Log, All);
 
-#if WITH_LIVELINK_DISCOVERY_MANAGER_THREAD
-
 namespace LiveLinkHubConnectionManager
 {
+	template <typename T>
+	concept CHasAnnotations = requires(const T& t) { t.GetAnnotations(); };
+
 	static void SendAnalyticsConnectionEstablished()
 	{
 		if (!FEngineAnalytics::IsAvailable())
@@ -40,14 +41,106 @@ namespace LiveLinkHubConnectionManager
 
 		FEngineAnalytics::GetProvider().RecordEvent(TEXT("Usage.LiveLinkHub.ConnectionEstablished"), {});
 	}
+
+	static ELiveLinkTopologyMode GetPollResultTopologyMode(const FProviderPollResultPtr& PollResult)
+	{
+		// Default to Hub since all LiveLinkHub instances were hubs before spokes were introduced.
+		ELiveLinkTopologyMode PollResultMode = ELiveLinkTopologyMode::Hub;
+
+		if (PollResult)
+		{
+			if (const FString* TopologyModeAnnotation = PollResult->Annotations.Find(FLiveLinkMessageAnnotation::TopologyModeAnnotation))
+			{
+				int64 TopologyModeValue = StaticEnum<ELiveLinkTopologyMode>()->GetValueByName(**TopologyModeAnnotation);
+				if (TopologyModeValue != INDEX_NONE)
+				{
+					PollResultMode = (ELiveLinkTopologyMode)TopologyModeValue;
+				}
+			}
+			else if (PollResult->Annotations.FindRef(FLiveLinkHubMessageAnnotation::ProviderTypeAnnotation) != UE::LiveLinkHub::Private::LiveLinkHubProviderType)
+			{
+				// Non-Hub livelink providers are usually external if they don't have annotations.
+				PollResultMode = ELiveLinkTopologyMode::External;
+			}
+		}
+
+		return PollResultMode;
+	}
+
+	template<CHasAnnotations T>
+	static bool CanConnectTo(const FString& MachineName, const T& ObjectWithAnnotations, const FLiveLinkHubInstanceId& InstanceId)
+	{
+		const TMap<FName, FString>& Annotations = ObjectWithAnnotations.GetAnnotations();
+		ELiveLinkHubAutoConnectMode AutoConnectMode = ELiveLinkHubAutoConnectMode::All;
+		if (const FString* AutoConnectModeAnnotation = Annotations.Find(FLiveLinkHubMessageAnnotation::AutoConnectModeAnnotation))
+		{
+			int64 AutoConnectModeValue = StaticEnum<ELiveLinkHubAutoConnectMode>()->GetValueByName(**AutoConnectModeAnnotation);
+			if (AutoConnectModeValue != INDEX_NONE)
+			{
+				AutoConnectMode = (ELiveLinkHubAutoConnectMode)AutoConnectModeValue;
+			}
+		}
+
+		// Prevent connecting to itself.
+		bool bSameInstance = false;
+		if (const FString* InstanceIdAnnotation = Annotations.Find(FLiveLinkHubMessageAnnotation::IdAnnotation))
+		{
+			bSameInstance = *InstanceIdAnnotation == InstanceId.ToString();
+		}
+
+		const bool bSameHost = MachineName == FPlatformProcess::ComputerName();
+
+		bool bAutoConnectMatchResult = !bSameInstance &&
+			(AutoConnectMode == ELiveLinkHubAutoConnectMode::All || (AutoConnectMode == ELiveLinkHubAutoConnectMode::LocalOnly && bSameHost));
+		if (!bAutoConnectMatchResult)
+		{
+			const FText AutoConnectModeName = UEnum::GetDisplayValueAsText(AutoConnectMode);
+			UE_LOG(LogLiveLinkHubConnectionManager, Verbose, TEXT("Refusing connection from incoming instance since it was in mode: %s"), *AutoConnectModeName.ToString());
+		}
+
+		return bAutoConnectMatchResult;
+	}
+
+	/** Returns whether this connection manager should accept connection requests from this poll result. */
+	static bool ShouldAcceptConnectionFrom(ELiveLinkTopologyMode InHostMode, const FProviderPollResultPtr& InPollResult, const FLiveLinkHubInstanceId& InstanceId)
+	{
+		if (!InPollResult)
+		{
+			return false;
+		}
+
+		// Topology Mode
+		ELiveLinkTopologyMode IncomingMode = LiveLinkHubConnectionManager::GetPollResultTopologyMode(InPollResult);
+
+		const bool bCompatibleModeResult = UE::LiveLink::Messaging::CanReceiveFrom(InHostMode, IncomingMode);
+
+		const FText ModeName = UEnum::GetDisplayValueAsText(InHostMode);
+		const FText IncomingModeName = UEnum::GetDisplayValueAsText(IncomingMode);
+		
+		if (!bCompatibleModeResult)
+		{
+			UE_LOG(LogLiveLinkHubConnectionManager, Verbose, TEXT("Refusing connection from incoming instance in %s mode. This app is in %s mode."), *IncomingModeName.ToString(), *ModeName.ToString());
+		}
+
+		return bCompatibleModeResult && LiveLinkHubConnectionManager::CanConnectTo(InPollResult->MachineName, *InPollResult, InstanceId);
+	}
 }
 
+#if WITH_LIVELINK_DISCOVERY_MANAGER_THREAD
+
+DECLARE_DELEGATE_RetVal(ELiveLinkTopologyMode, FGetTopologyMode);
+DECLARE_DELEGATE_RetVal(FLiveLinkHubInstanceId, FGetInstanceId);
 
 /** This utitlity is meant to be run on an unreal engine instance to look for livelink hub connections and to automatically create the message bus source for it. */
 class FLiveLinkHubConnectionManager
 {
 public:
-	FLiveLinkHubConnectionManager()
+	DECLARE_DELEGATE_RetVal(ELiveLinkTopologyMode, FOnGetTopologyMode);
+	DECLARE_DELEGATE_RetVal(FLiveLinkHubInstanceId, FOnGetInstanceId);
+
+	FLiveLinkHubConnectionManager(ELiveLinkTopologyMode InMode, FOnGetTopologyMode OnGetTopologyMode, FOnGetInstanceId OnGetInstanceId)
+		: GetTopologyModeDelegate(OnGetTopologyMode)
+		, GetInstanceIdDelegate(OnGetInstanceId)
 	{
 		FCoreUObjectDelegates::PostLoadMapWithWorld.AddRaw(this, &FLiveLinkHubConnectionManager::PostLoadMap);
 		FCoreDelegates::OnPostEngineInit.AddRaw(this, &FLiveLinkHubConnectionManager::StartDiscovery);
@@ -103,35 +196,20 @@ private:
 	void LookForLiveLinkHubConnection()
 	{
 		// Only look for a source if we don't have a valid connection.
-		if (!LastAddedSource.Value.IsValid() || !LastAddedSource.Value.Pin()->IsSourceStillValid())
+		UE_LOG(LogLiveLinkHubConnectionManager, Verbose, TEXT("Polling discovery results."));
+
+		TArray<FProviderPollResultPtr> PollResults = ILiveLinkModule::Get().GetMessageBusDiscoveryManager().GetDiscoveryResults();
+		for (const FProviderPollResultPtr& PollResult : PollResults)
 		{
-			UE_LOG(LogLiveLinkHubConnectionManager, Verbose, TEXT("Polling discovery results."));
+			const FString* ProviderType = PollResult->Annotations.Find(FLiveLinkHubMessageAnnotation::ProviderTypeAnnotation);
 
-			TArray<FProviderPollResultPtr> PollResults = ILiveLinkModule::Get().GetMessageBusDiscoveryManager().GetDiscoveryResults();
-			for (const FProviderPollResultPtr& PollResult : PollResults)
+			if (ProviderType && *ProviderType == UE::LiveLinkHub::Private::LiveLinkHubProviderType)
 			{
-				const FString* ProviderType = PollResult->Annotations.Find(FLiveLinkHubMessageAnnotation::ProviderTypeAnnotation);
-		
-				// Before this annotation was added, LiveLinkHub would automatically be added, so we need to keep the previous behavior if we discovered an older LiveLinkHub instance.
-				ELiveLinkHubAutoConnectMode AutoDiscover = ELiveLinkHubAutoConnectMode::Disabled;
-
-				if (const FString* AutoDiscoverAnnotation = PollResult->Annotations.Find(FLiveLinkHubMessageAnnotation::AutoDiscoverAnnotation))
-				{
-					int64 AutoDiscoverValue = StaticEnum<ELiveLinkHubAutoConnectMode>()->GetValueByName(**AutoDiscoverAnnotation);
-					if (AutoDiscoverValue != INDEX_NONE)
-					{
-						AutoDiscover = (ELiveLinkHubAutoConnectMode)AutoDiscoverValue;
-					}
-				}
-
-				const bool bSameHost = PollResult->MachineName == FPlatformProcess::ComputerName();
-
-				if (ProviderType && *ProviderType == UE::LiveLinkHub::Private::LiveLinkHubProviderType)
+				const ELiveLinkTopologyMode HostMode = GetTopologyModeDelegate.Execute();
+				const FLiveLinkHubInstanceId InstanceId = GetInstanceIdDelegate.Execute();
+				if (LiveLinkHubConnectionManager::ShouldAcceptConnectionFrom(HostMode, PollResult, InstanceId))
 				{
-					if (AutoDiscover == ELiveLinkHubAutoConnectMode::All || (AutoDiscover == ELiveLinkHubAutoConnectMode::LocalOnly && bSameHost))
-					{
-						AddLiveLinkSource(PollResult);
-					}
+					AddLiveLinkSource(PollResult);
 				}
 			}
 		}
@@ -150,28 +228,18 @@ private:
 
 			for (const FGuid& SourceId : LiveLinkClient->GetSources())
 			{
-				if (LiveLinkClient->GetSourceMachineName(SourceId).ToString() == PollResult->MachineName)
+				if (LiveLinkClient->GetSourceType(SourceId).ToString() == PollResult->Name && LiveLinkClient->GetSourceMachineName(SourceId).ToString() == PollResult->MachineName)
 				{
-					LiveLinkClient->RemoveSource(SourceId);
+					// Todo: Should we re-enable stale connections? 
+					UE_LOG(LogLiveLinkHubConnectionManager, Verbose, TEXT("Rejecting poll result since source %s already exists."), *PollResult->Name);
+					return;
 				}
 			}
 
-			TSharedPtr<ILiveLinkSource> LiveLinkSource = MakeShared<FLiveLinkHubMessageBusSource>(FText::FromString(PollResult->Name), FText::FromString(PollResult->MachineName), PollResult->Address, PollResult->MachineTimeOffset);
+			ILiveLinkHubMessagingModule& HubMessagingModule = FModuleManager::GetModuleChecked<ILiveLinkHubMessagingModule>("LiveLinkHubMessaging");
+			TSharedPtr<ILiveLinkSource> LiveLinkSource = MakeShared<FLiveLinkHubMessageBusSource>(FText::FromString(PollResult->Name), FText::FromString(PollResult->MachineName), PollResult->Address, PollResult->MachineTimeOffset, HubMessagingModule.GetInstanceId());
 			FGuid SourceId = LiveLinkClient->AddSource(LiveLinkSource);
 
-			// Remove the old source if we're re-establising a connection.
-			if (TSharedPtr<ILiveLinkSource> OldSource = LastAddedSource.Value.Pin())
-			{
-				// Todo: Handle one computer running multiple livelink hubs
-				if (OldSource->IsSourceStillValid() && OldSource->GetSourceMachineName().ToString() == PollResult->MachineName)
-				{
-					UE_LOG(LogLiveLinkHubConnectionManager, Verbose, TEXT("Removing disconnected livelink hub source."));
-					LiveLinkClient->RemoveSource(OldSource);
-				}
-			}
-
-			LastAddedSource = TPair<FMessageAddress, TWeakPtr<ILiveLinkSource>>{ PollResult->Address, LiveLinkSource };
-			ILiveLinkHubMessagingModule& HubMessagingModule = FModuleManager::GetModuleChecked<ILiveLinkHubMessagingModule>("LiveLinkHubMessaging");
 			HubMessagingModule.OnConnectionEstablished().Broadcast(SourceId);
 			LiveLinkHubConnectionManager::SendAnalyticsConnectionEstablished();
 		}
@@ -190,8 +258,10 @@ private:
 private:
 	/** Handle to the timer used to check for livelink hub providers. */
 	FTimerHandle ConnectionUpdateTimer;
-	/** Pointer to the message bus source used to know if we have a valid connection. */
-	TPair<FMessageAddress, TWeakPtr<ILiveLinkSource>> LastAddedSource;
+	/** Get the mode for this connection manager. */
+	FOnGetTopologyMode GetTopologyModeDelegate;
+	/** Get the instance id (Only relevant if this is running inside of livelinkhub. */
+	FOnGetInstanceId GetInstanceIdDelegate;
 };
 #else
 class FLiveLinkHubConnectionManager
diff --git a/Engine/Plugins/Animation/LiveLinkHub/Source/LiveLinkHubMessaging/Private/LiveLinkHubMessageBusSource.cpp b/Engine/Plugins/Animation/LiveLinkHub/Source/LiveLinkHubMessaging/Private/LiveLinkHubMessageBusSource.cpp
--- a/Engine/Plugins/Animation/LiveLinkHub/Source/LiveLinkHubMessaging/Private/LiveLinkHubMessageBusSource.cpp
+++ b/Engine/Plugins/Animation/LiveLinkHub/Source/LiveLinkHubMessaging/Private/LiveLinkHubMessageBusSource.cpp
@@ -6,10 +6,12 @@
 #include "Engine/Level.h"
 #include "Engine/SystemTimeTimecodeProvider.h"
 #include "Engine/World.h"
+#include "Features/IModularFeatures.h"
 #include "HAL/PlatformProcess.h"
 #include "ILiveLinkClient.h"
 #include "LiveLinkClient.h"
 #include "ILiveLinkModule.h"
+#include "LiveLinkHubMessagingModule.h"
 #include "LiveLinkHubMessages.h"
 #include "LiveLinkMessages.h"
 #include "LiveLinkTimecodeProvider.h"
@@ -25,8 +27,9 @@
 
 #include <limits>
 
-FLiveLinkHubMessageBusSource::FLiveLinkHubMessageBusSource(const FText& InSourceType, const FText& InSourceMachineName, const FMessageAddress& InConnectionAddress, double InMachineTimeOffset)
+FLiveLinkHubMessageBusSource::FLiveLinkHubMessageBusSource(const FText& InSourceType, const FText& InSourceMachineName, const FMessageAddress& InConnectionAddress, double InMachineTimeOffset, FLiveLinkHubInstanceId InInstanceId)
 	: FLiveLinkMessageBusSource(InSourceType, InSourceMachineName, InConnectionAddress, InMachineTimeOffset)
+	, InstanceId(InInstanceId)
 {
 #if WITH_EDITOR
 	if (FLevelEditorModule* LevelEditorModule = FModuleManager::GetModulePtr<FLevelEditorModule>("LevelEditor"))
@@ -54,6 +57,8 @@ double FLiveLinkHubMessageBusSource::GetDeadSourceTimeout() const
 
 void FLiveLinkHubMessageBusSource::SendConnectMessage()
 {
+	UE_LOG(LogLiveLinkHubMessaging, Verbose, TEXT("MessageBusSource (%s): Sending connect message to %s"), *MessageEndpoint->GetAddress().ToString(), *ConnectionAddress.ToString());
+
 	FLiveLinkHubConnectMessage* ConnectMessage = FMessageEndpoint::MakeMessage<FLiveLinkHubConnectMessage>();
 	ConnectMessage->ClientInfo = CreateLiveLinkClientInfo();
 
@@ -62,6 +67,20 @@ void FLiveLinkHubMessageBusSource::SendConnectMessage()
 	bIsValid = true;
 }
 
+bool FLiveLinkHubMessageBusSource::RequestSourceShutdown()
+{
+	UE_LOG(LogLiveLinkHubMessaging, Verbose, TEXT("MessageBusSource(%s) : Sending disconnect message to %s"), *MessageEndpoint->GetAddress().ToString(), *ConnectionAddress.ToString());
+
+	// Don't send the disconnect message if this source is shutting down because it's handling a disconnect message.
+	if (!bDisconnecting)
+	{
+		// Otherwise, inform our connected provider that it should remove this source.
+		SendMessage(FMessageEndpoint::MakeMessage<FLiveLinkHubDisconnectMessage>(GetSourceType().ToString(), GetSourceMachineName().ToString()));
+	}
+
+	return FLiveLinkMessageBusSource::RequestSourceShutdown();
+}
+
 void FLiveLinkHubMessageBusSource::SendClientInfoMessage()
 {
 	SendMessage(FMessageEndpoint::MakeMessage<FLiveLinkClientInfoMessage>(CreateLiveLinkClientInfo()));
@@ -93,7 +112,8 @@ void FLiveLinkHubMessageBusSource::InitializeMessageEndpoint(FMessageEndpointBui
 	FLiveLinkMessageBusSource::InitializeMessageEndpoint(EndpointBuilder);
 	EndpointBuilder
 		.Handling<FLiveLinkHubTimecodeSettings>(this, &FLiveLinkHubMessageBusSource::HandleTimecodeSettings)
-		.Handling<FLiveLinkHubCustomTimeStepSettings>(this, &FLiveLinkHubMessageBusSource::HandleCustomTimeStepSettings);
+		.Handling<FLiveLinkHubCustomTimeStepSettings>(this, &FLiveLinkHubMessageBusSource::HandleCustomTimeStepSettings)
+		.Handling<FLiveLinkHubDisconnectMessage>(this, &FLiveLinkHubMessageBusSource::HandleDisconnectMessage);
 }
 
 
@@ -111,6 +131,17 @@ void FLiveLinkHubMessageBusSource::HandleTimecodeSettings(const FLiveLinkHubTime
 	});
 }
 
+void FLiveLinkHubMessageBusSource::HandleDisconnectMessage(const FLiveLinkHubDisconnectMessage& Message, const TSharedRef<IMessageContext, ESPMode::ThreadSafe>& Context)
+{
+	bDisconnecting = true;
+
+	FGuid SourceId = SourceGuid;
+	ExecuteOnGameThread(UE_SOURCE_LOCATION, [SourceId]() {
+		ILiveLinkClient* Client = &IModularFeatures::Get().GetModularFeature<ILiveLinkClient>(ILiveLinkClient::ModularFeatureName);
+		Client->RemoveSource(SourceId);
+	});
+}
+
 FLiveLinkClientInfoMessage FLiveLinkHubMessageBusSource::CreateLiveLinkClientInfo() const
 {
 	FLiveLinkClientInfoMessage ClientInfo;
@@ -127,6 +158,17 @@ FLiveLinkClientInfoMessage FLiveLinkHubMessageBusSource::CreateLiveLinkClientInf
 	ClientInfo.Hostname = FPlatformProcess::ComputerName();
 	ClientInfo.ProjectName = FApp::GetProjectName();
 	ClientInfo.CurrentLevel = CurrentLevelName;
+
+	FLiveLinkHubMessagingModule& MessagingModule = FModuleManager::GetModuleChecked<FLiveLinkHubMessagingModule>("LiveLinkHubMessaging");
+	ELiveLinkTopologyMode Mode = MessagingModule.GetHostTopologyMode();
+
+	// Only populate this field if this is a Hub.
+	if (Mode == ELiveLinkTopologyMode::Hub)
+	{
+		ClientInfo.LiveLinkInstanceName = InstanceId.ToString();
+		ClientInfo.TopologyMode = Mode;
+	}
+
 	ClientInfo.LiveLinkVersion = ILiveLinkClient::LIVELINK_VERSION;
 
 	return ClientInfo;
diff --git a/Engine/Plugins/Animation/LiveLinkHub/Source/LiveLinkHubMessaging/Private/LiveLinkHubMessageBusSource.h b/Engine/Plugins/Animation/LiveLinkHub/Source/LiveLinkHubMessaging/Private/LiveLinkHubMessageBusSource.h
--- a/Engine/Plugins/Animation/LiveLinkHub/Source/LiveLinkHubMessaging/Private/LiveLinkHubMessageBusSource.h
+++ b/Engine/Plugins/Animation/LiveLinkHub/Source/LiveLinkHubMessaging/Private/LiveLinkHubMessageBusSource.h
@@ -2,7 +2,9 @@
 
 #pragma once
 
+#include "ILiveLinkHubMessagingModule.h"
 #include "LiveLinkHubMessages.h"
+#include "LiveLinkHubMessageBusSourceSettings.h"
 #include "LiveLinkMessageBusSource.h"
 
 class UWorld;
@@ -15,7 +17,7 @@ struct FLiveLinkHubTimecodeSettings;
 class FLiveLinkHubMessageBusSource : public FLiveLinkMessageBusSource
 {
 public:
-	FLiveLinkHubMessageBusSource(const FText& InSourceType, const FText& InSourceMachineName, const FMessageAddress& InConnectionAddress, double InMachineTimeOffset);
+	FLiveLinkHubMessageBusSource(const FText& InSourceType, const FText& InSourceMachineName, const FMessageAddress& InConnectionAddress, double InMachineTimeOffset, FLiveLinkHubInstanceId InInstanceId);
 	virtual ~FLiveLinkHubMessageBusSource();
 
 protected:
@@ -24,6 +26,8 @@ protected:
 	virtual void InitializeMessageEndpoint(FMessageEndpointBuilder& EndpointBuilder);
 	virtual double GetDeadSourceTimeout() const override;
 	virtual void SendConnectMessage() override;
+	virtual TSubclassOf<ULiveLinkSourceSettings> GetSettingsClass() const override { return ULiveLinkHubMessageBusSourceSettings::StaticClass(); }
+	virtual bool RequestSourceShutdown() override;
 	//~ End FLiveLinkMessageBusSource
 
 private:
@@ -31,11 +35,16 @@ private:
 	void HandleCustomTimeStepSettings(const FLiveLinkHubCustomTimeStepSettings& Message, const TSharedRef<IMessageContext, ESPMode::ThreadSafe>& Context);
 	/** Handle a timecode settings message and update the engine's timecode settings accordingly. */
 	void HandleTimecodeSettings(const FLiveLinkHubTimecodeSettings& Message, const TSharedRef<IMessageContext, ESPMode::ThreadSafe>& Context);
-
+	/** Handle a disconnect request from the provider. */
+	void HandleDisconnectMessage(const FLiveLinkHubDisconnectMessage& Message, const TSharedRef<IMessageContext, ESPMode::ThreadSafe>& Context);
 	/** Send information about this UE client to the connected provider */
 	void SendClientInfoMessage();
 	/** Handler called on map changed to update the livelink hub. */
 	void OnMapChanged(UWorld* World, EMapChangeType ChangeType);
 	/** Gather information about this client to put in a client info struct. */
 	FLiveLinkClientInfoMessage CreateLiveLinkClientInfo() const;
+	/** ID of this instance if running inside LiveLinkHub. */
+	const FLiveLinkHubInstanceId InstanceId;
+	/** If this source is in the process of disconnecting. */
+	bool bDisconnecting = false;
 };
diff --git a/Engine/Plugins/Animation/LiveLinkHub/Source/LiveLinkHubMessaging/Private/LiveLinkHubMessageBusSourceFactory.cpp b/Engine/Plugins/Animation/LiveLinkHub/Source/LiveLinkHubMessaging/Private/LiveLinkHubMessageBusSourceFactory.cpp
--- a/Engine/Plugins/Animation/LiveLinkHub/Source/LiveLinkHubMessaging/Private/LiveLinkHubMessageBusSourceFactory.cpp
+++ b/Engine/Plugins/Animation/LiveLinkHub/Source/LiveLinkHubMessaging/Private/LiveLinkHubMessageBusSourceFactory.cpp
@@ -1,8 +1,11 @@
 // Copyright Epic Games, Inc. All Rights Reserved.
 
 #include "LiveLinkHubMessageBusSourceFactory.h"
+
 #include "LiveLinkHubMessageBusSource.h"
+#include "ILiveLinkHubMessagingModule.h"
 #include "Misc/ConfigCacheIni.h"
+#include "Modules/ModuleManager.h"
 #include "SLiveLinkMessageBusSourceFactory.h"
 
 
@@ -31,12 +34,14 @@ TSharedPtr<FLiveLinkMessageBusSource> ULiveLinkHubMessageBusSourceFactory::MakeS
 																				   const FMessageAddress& Address,
 																				   double TimeOffset) const
 {
-	return MakeShared<FLiveLinkHubMessageBusSource>(Name, MachineName, Address, TimeOffset);
+	ILiveLinkHubMessagingModule& Module = FModuleManager::Get().GetModuleChecked<ILiveLinkHubMessagingModule>("LiveLinkHubMessaging");
+	return MakeShared<FLiveLinkHubMessageBusSource>(Name, MachineName, Address, TimeOffset, Module.GetInstanceId());
 }
 
 bool ULiveLinkHubMessageBusSourceFactory::IsEnabled() const
 {
-	return GConfig->GetBoolOrDefault(TEXT("LiveLinkHub"), TEXT("bEnableLLHMessageBusSourceFactory"), true, GEngineIni);
+	ILiveLinkHubMessagingModule& Module = FModuleManager::Get().GetModuleChecked<ILiveLinkHubMessagingModule>("LiveLinkHubMessaging");
+	return Module.GetHostTopologyMode() == ELiveLinkTopologyMode::Hub || Module.GetHostTopologyMode() == ELiveLinkTopologyMode::UnrealClient;
 }
 
 
diff --git a/Engine/Plugins/Animation/LiveLinkHub/Source/LiveLinkHubMessaging/Private/LiveLinkHubMessages.cpp b/Engine/Plugins/Animation/LiveLinkHub/Source/LiveLinkHubMessaging/Private/LiveLinkHubMessages.cpp
--- a/Engine/Plugins/Animation/LiveLinkHub/Source/LiveLinkHubMessaging/Private/LiveLinkHubMessages.cpp
+++ b/Engine/Plugins/Animation/LiveLinkHub/Source/LiveLinkHubMessaging/Private/LiveLinkHubMessages.cpp
@@ -12,11 +12,13 @@
 
 const FName UE::LiveLinkHub::Private::LiveLinkHubProviderType = TEXT("LiveLinkHub");
 FName FLiveLinkHubMessageAnnotation::ProviderTypeAnnotation = TEXT("ProviderType");
-FName FLiveLinkHubMessageAnnotation::AutoDiscoverAnnotation = TEXT("AutoDiscover");
+FName FLiveLinkHubMessageAnnotation::AutoConnectModeAnnotation = TEXT("AutoConnect");
+FName FLiveLinkHubMessageAnnotation::IdAnnotation = TEXT("Id");
 
 
 DEFINE_LOG_CATEGORY_STATIC(LogLiveLinkHubMessages, Log, All);
 
+
 namespace UE::LiveLinkHubMessages::Private
 {
 	TOptional<FLiveLinkSubjectKey> FindSubjectKey(ILiveLinkClient* LiveLinkClient, FName SubjectName)
@@ -73,6 +75,10 @@ void FLiveLinkHubCustomTimeStepSettings::AssignCustomTimeStepToEngine() const
 	{
 		NewCustomTimeStep->SubjectKey = *Target;
 	}
+	else
+	{
+		NewCustomTimeStep->SubjectKey = FLiveLinkSubjectKey{};
+	}
 
 	// Override the custom timestep for the engine.
 	GEngine->SetCustomTimeStep(NewCustomTimeStep);
diff --git a/Engine/Plugins/Animation/LiveLinkHub/Source/LiveLinkHubMessaging/Private/LiveLinkHubMessagingModule.cpp b/Engine/Plugins/Animation/LiveLinkHub/Source/LiveLinkHubMessaging/Private/LiveLinkHubMessagingModule.cpp
--- a/Engine/Plugins/Animation/LiveLinkHub/Source/LiveLinkHubMessaging/Private/LiveLinkHubMessagingModule.cpp
+++ b/Engine/Plugins/Animation/LiveLinkHub/Source/LiveLinkHubMessaging/Private/LiveLinkHubMessagingModule.cpp
@@ -1,108 +1,150 @@
 // Copyright Epic Games, Inc. All Rights Reserved.
 
-#include "ILiveLinkHubMessagingModule.h"
+#include "LiveLinkHubMessagingModule.h"
 
 #include "CoreMinimal.h"
 #include "LiveLinkHubConnectionManager.h"
 #include "LiveLinkHubMessageBusSourceFactory.h"
 #include "LiveLinkMessageBusSourceFactory.h"
+#include "MessageEndpoint.h"
+#include "MessageEndpointBuilder.h"
 #include "Misc/ConfigCacheIni.h"
+#include "Misc/ScopeLock.h"
 #include "Modules/ModuleManager.h"
 
-#ifndef WITH_LIVELINK_DISCOVERY_MANAGER_THREAD
-#define WITH_LIVELINK_DISCOVERY_MANAGER_THREAD 1
-#endif
-
+DEFINE_LOG_CATEGORY(LogLiveLinkHubMessaging);
 
-class FLiveLinkHubMessagingModule : public ILiveLinkHubMessagingModule
+namespace Internal
 {
-public:
-	//~ Begin IModuleInterface
-	virtual void StartupModule() override
+	// Source: https://en.cppreference.com/w/cpp/utility/variant/visit
+	template<class... Ts>
+	struct TOverloaded : Ts...
 	{
-		// The connection manager is used to communicate with the hub,
-		// so we don't need it when we're running the hub itself.
-		bUseConnectionManager = GConfig->GetBoolOrDefault(
-			TEXT("LiveLinkHub"), TEXT("bUseMessagingConnectionManager"), true, GEngineIni);
+		using Ts::operator()...;
+	};
+
+	template<class... Ts> TOverloaded(Ts...) -> TOverloaded<Ts...>;
+}
+
+void FLiveLinkHubMessagingModule::StartupModule()
+{
+	const bool bIsLiveLinkHubHost = GConfig->GetBoolOrDefault(
+		TEXT("LiveLink"), TEXT("bCreateLiveLinkHubInstance"), false, GEngineIni);
+
+	InstanceInfo.TopologyMode = bIsLiveLinkHubHost ? ELiveLinkTopologyMode::Hub : ELiveLinkTopologyMode::UnrealClient;
 
 #if WITH_LIVELINK_DISCOVERY_MANAGER_THREAD
-		if (bUseConnectionManager)
-		{
-			ConnectionManager = MakeUnique<FLiveLinkHubConnectionManager>();
-		}
+	
+	ConnectionManager = MakePimpl<FLiveLinkHubConnectionManager>(InstanceInfo.TopologyMode,
+		FLiveLinkHubConnectionManager::FOnGetTopologyMode::CreateRaw(this, &FLiveLinkHubMessagingModule::GetHostTopologyMode),
+		FLiveLinkHubConnectionManager::FOnGetInstanceId::CreateRaw(this, &FLiveLinkHubMessagingModule::GetInstanceId)
+	);
+
+	GameThreadEndpoint = FMessageEndpoint::Builder(TEXT("LiveLinkHubDiscoveryEndpoint"))
+		.Handling<FLiveLinkHubDiscoveryMessage>(this, &FLiveLinkHubMessagingModule::HandleDiscoveryMessage);
+
+
+	GameThreadEndpoint->Subscribe<FLiveLinkHubDisconnectMessage>();
+	GameThreadEndpoint->Subscribe<FLiveLinkHubDiscoveryMessage>();
 #endif
 
-		SourceFilterDelegate = ILiveLinkModule::Get().RegisterMessageBusSourceFilter(FOnLiveLinkShouldDisplaySource::CreateRaw(this, &FLiveLinkHubMessagingModule::OnFilterMessageBusSource));
-	}
+	SourceFilterDelegate = ILiveLinkModule::Get().RegisterMessageBusSourceFilter(FOnLiveLinkShouldDisplaySource::CreateRaw(this, &FLiveLinkHubMessagingModule::OnFilterMessageBusSource));
+}
 
-	virtual void ShutdownModule() override
+void FLiveLinkHubMessagingModule::ShutdownModule()
+{ 
+	if (ILiveLinkModule* LiveLinkModule = FModuleManager::Get().GetModulePtr<ILiveLinkModule>("LiveLink"))
 	{
-		if (ILiveLinkModule* LiveLinkModule = FModuleManager::Get().GetModulePtr<ILiveLinkModule>("LiveLink"))
-		{
-			LiveLinkModule->UnregisterMessageBusSourceFilter(SourceFilterDelegate);
-		}
+		LiveLinkModule->UnregisterMessageBusSourceFilter(SourceFilterDelegate);
+	}
 
 #if WITH_LIVELINK_DISCOVERY_MANAGER_THREAD
-		ConnectionManager.Reset();
+	ConnectionManager.Reset();
 #endif
-	}
-	//~ End IModuleInterface
+}
 
-	//~ ILiveLinkHubMessagingModule interface
-	virtual FOnHubConnectionEstablished& OnConnectionEstablished() override
-	{
-		return ConnectionEstablishedDelegate;
-	}
+void FLiveLinkHubMessagingModule::SetHostTopologyMode(ELiveLinkTopologyMode InMode)
+{
+	FScopeLock Lock(&InstanceInfoLock);
+	InstanceInfo.TopologyMode = InMode;
+}
+
+FLiveLinkHubInstanceId FLiveLinkHubMessagingModule::GetInstanceId() const
+{
+	FScopeLock Lock(&InstanceInfoLock);
+	return InstanceInfo.Id;
+}
+
+void FLiveLinkHubMessagingModule::SetInstanceId(const FLiveLinkHubInstanceId& Id)
+{
+	FScopeLock Lock(&InstanceInfoLock);
+	InstanceInfo.Id = Id;
+}
+
+ELiveLinkTopologyMode FLiveLinkHubMessagingModule::GetHostTopologyMode() const
+{
+	FScopeLock Lock(&InstanceInfoLock);
+	return InstanceInfo.TopologyMode;
+}
 
-private:
-	/** Filter invoked by the messagebus source factory to filter out sources in the creation panel. */
-	bool OnFilterMessageBusSource(UClass* FactoryClass, FProviderPollResultPtr PollResult) const
+bool FLiveLinkHubMessagingModule::OnFilterMessageBusSource(UClass* FactoryClass, TSharedPtr<struct FProviderPollResult, ESPMode::ThreadSafe> PollResult)
+{
+	return LiveLinkHubConnectionManager::ShouldAcceptConnectionFrom(InstanceInfo.TopologyMode, PollResult, InstanceInfo.Id);
+}
+
+void FLiveLinkHubMessagingModule::HandleDiscoveryMessage(const FLiveLinkHubDiscoveryMessage& Message, const TSharedRef<IMessageContext, ESPMode::ThreadSafe>& Context)
+{
+	// Before this annotation was added, LiveLinkHub would automatically be added, so we need to keep the previous behavior if we discovered an older LiveLinkHub instance.
+	ELiveLinkHubAutoConnectMode AutoConnectMode = ELiveLinkHubAutoConnectMode::Disabled;
+
+	if (const FString* AutoConnectModeAnnotation = Context->GetAnnotations().Find(FLiveLinkHubMessageAnnotation::AutoConnectModeAnnotation))
 	{
-		if (!PollResult)
+		int64 AutoConnectModeValue = StaticEnum<ELiveLinkHubAutoConnectMode>()->GetValueByName(**AutoConnectModeAnnotation);
+		if (AutoConnectModeValue != INDEX_NONE)
 		{
-			return false;
+			AutoConnectMode = (ELiveLinkHubAutoConnectMode)AutoConnectModeValue;
 		}
+	}
 
-		const bool bIsLiveLinkHubProvider = PollResult->Annotations.FindRef(FLiveLinkHubMessageAnnotation::ProviderTypeAnnotation) == UE::LiveLinkHub::Private::LiveLinkHubProviderType;
+	bool bShouldConnect = UE::LiveLink::Messaging::CanReceiveFrom(InstanceInfo.TopologyMode, Message.Mode)
+		&& LiveLinkHubConnectionManager::CanConnectTo(Message.MachineName, *Context, InstanceInfo.Id);
+	if (bShouldConnect)
+	{
+		const double MachineTimeOffset = LiveLinkMessageBusHelper::CalculateProviderMachineOffset(Message.CreationTime, Context);
 
-		if (!bUseConnectionManager)
-		{
-			// When running the hub
-			if (FactoryClass == ULiveLinkMessageBusSourceFactory::StaticClass())
-			{
-				// Don't show the livelink hub in the message bus source discovery. (Should be changed in the future if we want to allow hubs to speak to each others).
-				return !bIsLiveLinkHubProvider;
-			}
-		}
-		else
-		{
-			// When running in UE
-			if (FactoryClass == ULiveLinkHubMessageBusSourceFactory::StaticClass())
-			{
-				return bIsLiveLinkHubProvider;
-			}
-			else if (FactoryClass == ULiveLinkMessageBusSourceFactory::StaticClass())
-			{
-				return !bIsLiveLinkHubProvider;
-			}
-		}
+		TSharedPtr<FLiveLinkHubMessageBusSource> LiveLinkSource = MakeShared<FLiveLinkHubMessageBusSource>(FText::FromString(Message.ProviderName), FText::FromString(Message.MachineName), Context->GetSender(), MachineTimeOffset, InstanceInfo.Id);
 
-		return true;
+		ILiveLinkClient* Client = &IModularFeatures::Get().GetModularFeature<ILiveLinkClient>(ILiveLinkClient::ModularFeatureName);
+		Client->AddSource(LiveLinkSource);
 	}
+}
 
-	bool bUseConnectionManager;
-
-#if WITH_LIVELINK_DISCOVERY_MANAGER_THREAD
-	/** Manages the connection to the live link hub. */
-	TUniquePtr<FLiveLinkHubConnectionManager> ConnectionManager;
-#endif
+FLiveLinkHubInstanceId::FLiveLinkHubInstanceId(FGuid Guid)
+{
+	Id.Set<FGuid>(Guid);
+}
 
-	/** Handle to the delegate used to filter message bus sources. */
-	FDelegateHandle SourceFilterDelegate;
+FLiveLinkHubInstanceId::FLiveLinkHubInstanceId(FStringView NamedId)
+{
+	Id.Emplace<FString>(NamedId);
+}
 
-	/** Delegate called when the connection between a livelink hub and the editor is established. */
-	FOnHubConnectionEstablished ConnectionEstablishedDelegate;
-};
+FString FLiveLinkHubInstanceId::ToString() const
+{
+	auto Overload = Internal::TOverloaded{
+		[](FGuid Id) -> FString 
+		{
+			TStringBuilder<20> LiveLinkHubName;
+			LiveLinkHubName << TEXT("Live Link Hub") << TEXT(" (") << *Id.ToString().Right(4).ToLower() << TEXT(")");
+			return LiveLinkHubName.ToString();
+		},
+		[](const FString& Id) -> FString 
+		{
+			return Id;
+		}
+	};
+	return Visit(Overload, Id);
+}
 
 
 IMPLEMENT_MODULE(FLiveLinkHubMessagingModule, LiveLinkHubMessaging);
diff --git a/Engine/Plugins/Animation/LiveLinkHub/Source/LiveLinkHubMessaging/Private/LiveLinkHubMessagingModule.h b/Engine/Plugins/Animation/LiveLinkHub/Source/LiveLinkHubMessaging/Private/LiveLinkHubMessagingModule.h
new file mode 100644
--- /dev/null
+++ b/Engine/Plugins/Animation/LiveLinkHub/Source/LiveLinkHubMessaging/Private/LiveLinkHubMessagingModule.h
@@ -0,0 +1,81 @@
+// Copyright Epic Games, Inc. All Rights Reserved.
+
+#pragma once
+
+#include "ILiveLinkHubMessagingModule.h"
+
+#include "LiveLinkHubMessages.h"
+#include "HAL/CriticalSection.h"
+
+#ifndef WITH_LIVELINK_DISCOVERY_MANAGER_THREAD
+#define WITH_LIVELINK_DISCOVERY_MANAGER_THREAD 1
+#endif
+
+class FLiveLinkHubConnectionManager;
+
+class FLiveLinkHubMessagingModule : public ILiveLinkHubMessagingModule
+{
+public:
+	//~ Begin IModuleInterface
+	virtual void StartupModule() override;
+
+	virtual void ShutdownModule() override;
+	//~ End IModuleInterface
+
+	//~ Begin ILiveLinkHubMessagingModule interface
+	virtual FOnHubConnectionEstablished& OnConnectionEstablished() override
+	{
+		return ConnectionEstablishedDelegate;
+	}
+
+	virtual void SetHostTopologyMode(ELiveLinkTopologyMode InMode) override;
+	virtual FLiveLinkHubInstanceId GetInstanceId() const override;
+	virtual void SetInstanceId(const FLiveLinkHubInstanceId& Id) override;
+	virtual ELiveLinkTopologyMode GetHostTopologyMode() const override;
+	//~ End ILiveLinkHubMessagingModule interface
+
+private:
+	/**
+	 * Note: Invoked on the UI (Game) thread.
+	 * Filter invoked by the messagebus source factory to filter out sources in the creation panel.
+	 */
+	bool OnFilterMessageBusSource(UClass* FactoryClass, TSharedPtr<struct FProviderPollResult, ESPMode::ThreadSafe> PollResult);
+
+	/** Handle a message telling this host to connect to a provider if the topology mode is compatible. */
+	void HandleDiscoveryMessage(const FLiveLinkHubDiscoveryMessage& Message, const TSharedRef<class IMessageContext, ESPMode::ThreadSafe>& Context);
+
+private:
+	bool bUseConnectionManager;
+
+#if WITH_LIVELINK_DISCOVERY_MANAGER_THREAD
+	/** Manages the connection to the live link hub. */
+	TPimplPtr<FLiveLinkHubConnectionManager> ConnectionManager;
+#endif
+
+	/** Handle to the delegate used to filter message bus sources. */
+	FDelegateHandle SourceFilterDelegate;
+
+	/** Delegate called when the connection between a livelink hub and the editor is established. */
+	FOnHubConnectionEstablished ConnectionEstablishedDelegate;
+
+	/** Lock to access the instance info struct. */
+	mutable FCriticalSection InstanceInfoLock;
+
+	struct FInstanceInfo
+	{
+		/** Topology Mode for this host. */
+		ELiveLinkTopologyMode TopologyMode = ELiveLinkTopologyMode::Hub;
+		/** Instance ID for this host. */
+		FLiveLinkHubInstanceId Id = FLiveLinkHubInstanceId(FGuid());
+	} InstanceInfo;
+
+	/** Critical section to protect access to the disconnected clients list. */
+	FCriticalSection DisconnectedClientsLock;
+
+	/** List of clients that were explicitly disconnected and shouldn't reconnect automatically. */
+	TMap<FString, double> RecentlyDisconnectedClients;
+
+	/** Simple endpoint meant to respond to LLH discovery messages. */
+	TSharedPtr<class FMessageEndpoint> GameThreadEndpoint;
+};
+
diff --git a/Engine/Plugins/Animation/LiveLinkHub/Source/LiveLinkHubMessaging/Public/ILiveLinkHubMessagingModule.h b/Engine/Plugins/Animation/LiveLinkHub/Source/LiveLinkHubMessaging/Public/ILiveLinkHubMessagingModule.h
--- a/Engine/Plugins/Animation/LiveLinkHub/Source/LiveLinkHubMessaging/Public/ILiveLinkHubMessagingModule.h
+++ b/Engine/Plugins/Animation/LiveLinkHub/Source/LiveLinkHubMessaging/Public/ILiveLinkHubMessagingModule.h
@@ -3,6 +3,8 @@
 #pragma once
 
 #include "Delegates/Delegate.h"
+#include "Misc/Guid.h"
+#include "Misc/TVariant.h"
 #include "Modules/ModuleInterface.h"
 
 #include "Delegates/DelegateCombinations.h"
@@ -10,9 +12,40 @@
 
 DECLARE_MULTICAST_DELEGATE_OneParam(FOnHubConnectionEstablished, FGuid SourceId);
 
+
+DECLARE_LOG_CATEGORY_EXTERN(LogLiveLinkHubMessaging, Display, All);
+
+
+class LIVELINKHUBMESSAGING_API FLiveLinkHubInstanceId
+{
+public:
+	explicit FLiveLinkHubInstanceId(FGuid Guid);
+	explicit FLiveLinkHubInstanceId(FStringView NamedId);
+
+	FString ToString() const;
+
+private:
+	/** Variant holding the actual instance ID. */
+	TVariant<FGuid, FString> Id;
+};
+
+enum class ELiveLinkTopologyMode : uint8;
+
 class LIVELINKHUBMESSAGING_API ILiveLinkHubMessagingModule : public IModuleInterface
 {
 public:
 	/** Delegate called when a connection is established to a livelink hub. */
 	virtual FOnHubConnectionEstablished& OnConnectionEstablished() = 0;
+
+	/** Set the topology mode for this host. This will dictate to what apps it can connect to. */
+	virtual void SetHostTopologyMode(ELiveLinkTopologyMode InMode) = 0;
+
+	/** Get the topology mode for this host. */
+	virtual ELiveLinkTopologyMode GetHostTopologyMode() const = 0;
+
+	/** Set the instance ID for this connection manager, used to detect if it's trying to connect to itself. */
+	virtual void SetInstanceId(const FLiveLinkHubInstanceId& Id) = 0;
+
+	/** Get the ID for this running instance. Returns an empty guid if this isn't running as part of LiveLinkHub. */
+	virtual FLiveLinkHubInstanceId GetInstanceId() const = 0;
 };
diff --git a/Engine/Plugins/Animation/LiveLinkHub/Source/LiveLinkHubMessaging/Public/LiveLinkHubMessageBusSourceSettings.h b/Engine/Plugins/Animation/LiveLinkHub/Source/LiveLinkHubMessaging/Public/LiveLinkHubMessageBusSourceSettings.h
new file mode 100644
--- /dev/null
+++ b/Engine/Plugins/Animation/LiveLinkHub/Source/LiveLinkHubMessaging/Public/LiveLinkHubMessageBusSourceSettings.h
@@ -0,0 +1,15 @@
+// Copyright Epic Games, Inc. All Rights Reserved.
+
+#pragma once
+
+#include "LiveLinkMessageBusSourceSettings.h"
+
+#include "LiveLinkHubMessageBusSourceSettings.generated.h"
+
+/** Class used to allow customizing the default behavior of LiveLinkHub sources and more generally to allow quickly identifying if a source is a LLH instance. */
+UCLASS()
+class LIVELINKHUBMESSAGING_API ULiveLinkHubMessageBusSourceSettings : public ULiveLinkMessageBusSourceSettings
+{
+public:
+	GENERATED_BODY()
+};
\ No newline at end of file
diff --git a/Engine/Plugins/Animation/LiveLinkHub/Source/LiveLinkHubMessaging/Public/LiveLinkHubMessages.h b/Engine/Plugins/Animation/LiveLinkHub/Source/LiveLinkHubMessaging/Public/LiveLinkHubMessages.h
--- a/Engine/Plugins/Animation/LiveLinkHub/Source/LiveLinkHubMessaging/Public/LiveLinkHubMessages.h
+++ b/Engine/Plugins/Animation/LiveLinkHub/Source/LiveLinkHubMessaging/Public/LiveLinkHubMessages.h
@@ -3,6 +3,8 @@
 #pragma once
 
 #include "GenlockedCustomTimeStep.h"
+#include "ILiveLinkHubMessagingModule.h"
+#include "ILiveLinkClient.h"
 #include "LiveLinkMessages.h"
 #include "LiveLinkTypes.h"
 #include "Misc/FrameRate.h"
@@ -32,8 +34,10 @@ struct LIVELINKHUBMESSAGING_API FLiveLinkHubMessageAnnotation
 	static FName ProviderTypeAnnotation;
 
 	/** Annotation to indicate if this source should be automatically added to the list of LiveLink sources. */
-	static FName AutoDiscoverAnnotation;
+	static FName AutoConnectModeAnnotation;
 
+	/** Instance ID annotation used to identify the running LLH instance. */
+	static FName IdAnnotation;
 };
 
 namespace UE::LiveLinkHub::Private
@@ -147,9 +151,17 @@ struct LIVELINKHUBMESSAGING_API FLiveLinkClientInfoMessage
 	UPROPERTY()
 	FString CurrentLevel;
 
+	/** If this is representing a LiveLinkHub instance in Hub mode, this holds the LiveLink provider name, otherwise it's empty. */
+	UPROPERTY()
+	FString LiveLinkInstanceName;
+
+	/** Whether the client is a hub or an unreal instance. */
+	UPROPERTY()
+	ELiveLinkTopologyMode TopologyMode = ELiveLinkTopologyMode::UnrealClient;
+
 	/** LiveLink Version in use by this client. */
 	UPROPERTY()
-	int32 LiveLinkVersion = 1;
+	int32 LiveLinkVersion = ILiveLinkClient::LIVELINK_VERSION;
 };
 
 /** Special connection message used when connecting to a livelink hub that contains information about this client. */
@@ -162,3 +174,61 @@ struct LIVELINKHUBMESSAGING_API FLiveLinkHubConnectMessage
 	UPROPERTY()
 	FLiveLinkClientInfoMessage ClientInfo;
 };
+
+/** Special connection message used to tell a UE client or Hub that they should disconnect themselves. */
+USTRUCT()
+struct LIVELINKHUBMESSAGING_API FLiveLinkHubDisconnectMessage
+{
+	GENERATED_BODY()
+
+	/** Name of the provider to disconnect. */
+	UPROPERTY()	
+	FString ProviderName;
+
+	/** Name of the machine that hosts the provider. */
+	UPROPERTY()
+	FString MachineName;
+};
+
+/** Discovery message used by LiveLinkHubConnectionManager to find providers to connect to. */
+USTRUCT()
+struct LIVELINKHUBMESSAGING_API FLiveLinkHubDiscoveryMessage
+{
+	GENERATED_BODY()
+
+	FLiveLinkHubDiscoveryMessage() = default;
+
+	FLiveLinkHubDiscoveryMessage(FString InProviderName, ELiveLinkTopologyMode InMode, const FLiveLinkHubInstanceId& InInstanceId)
+		: ProviderName(MoveTemp(InProviderName))
+		, Mode(InMode)
+		, InstanceId(InInstanceId.ToString())
+	{
+	}
+
+	/** Name of the provider to connect. */
+	UPROPERTY()
+	FString ProviderName;
+
+	/** Name of the provider to connect. */
+	UPROPERTY()
+	ELiveLinkTopologyMode Mode = ELiveLinkTopologyMode::Hub;
+
+	/** Name of the machine that hosts the provider. */
+	UPROPERTY()
+	FString MachineName = FPlatformProcess::ComputerName();
+
+	/** Unique ID for this provider. */
+	UPROPERTY()
+	FString InstanceId;
+
+	/** Creation time used to calculate the machine time offset. */
+	UPROPERTY()
+	double CreationTime = FPlatformTime::Seconds();
+
+	/** LiveLink Version in use by this client. */
+	UPROPERTY()
+	int32 LiveLinkVersion = ILiveLinkClient::LIVELINK_VERSION;
+};
+
+
+
diff --git a/Engine/Source/Programs/LiveLinkHub/Config/DefaultEngine.ini b/Engine/Source/Programs/LiveLinkHub/Config/DefaultEngine.ini
--- a/Engine/Source/Programs/LiveLinkHub/Config/DefaultEngine.ini
+++ b/Engine/Source/Programs/LiveLinkHub/Config/DefaultEngine.ini
@@ -7,8 +7,8 @@ bEnableParentSubjects=true
 bCreateLiveLinkHubInstance=true
 
 [LiveLinkHub]
-bEnableLLHMessageBusSourceFactory=false
-bUseMessagingConnectionManager=false
+bEnableLLHMessageBusSourceFactory=true
+bUseMessagingConnectionManager=true
 
 [/Script/MacTargetPlatform.XcodeProjectSettings]
 TemplateMacPlist=(FilePath="/Engine/Source/Programs/LiveLinkHub/Resources/Mac/Info.plist")
diff --git a/Engine/Source/Runtime/LiveLinkMessageBusFramework/Private/LiveLinkProvider.cpp b/Engine/Source/Runtime/LiveLinkMessageBusFramework/Private/LiveLinkProvider.cpp
--- a/Engine/Source/Runtime/LiveLinkMessageBusFramework/Private/LiveLinkProvider.cpp
+++ b/Engine/Source/Runtime/LiveLinkMessageBusFramework/Private/LiveLinkProvider.cpp
@@ -25,6 +25,52 @@ template TSharedPtr<ILiveLinkProvider> ILiveLinkProvider::CreateLiveLinkProvider
 FName FLiveLinkMessageAnnotation::SubjectAnnotation = TEXT("SubjectName");
 FName FLiveLinkMessageAnnotation::RoleAnnotation = TEXT("Role");
 FName FLiveLinkMessageAnnotation::OriginalSourceAnnotation = TEXT("OriginalSource");
+FName FLiveLinkMessageAnnotation::TopologyModeAnnotation = TEXT("TopologyMode");
+
+
+namespace UE::LiveLink::Messaging
+{
+	using enum ELiveLinkTopologyMode;
+
+	/**
+	 * Map of possible connections.
+	 * Left side can receive a connection from right side.
+	 */
+	static const TMultiMap<ELiveLinkTopologyMode, ELiveLinkTopologyMode> CanReceiveFrom_Map =
+	{
+		{ UnrealClient, Hub },
+		{ UnrealClient, UnrealClient },
+		{ UnrealClient, External },
+		{ Spoke, External },
+		{ Hub, External },
+		{ Hub, Spoke },
+	};
+
+	/**
+	 * Map of possible connections.
+	 * Left side can transmit data to a client from right side.
+	 */
+	static const TMultiMap<ELiveLinkTopologyMode, ELiveLinkTopologyMode> CanTransmitTo_Map =
+	{
+		{ Hub, UnrealClient },
+		{ UnrealClient, UnrealClient },
+		{ Spoke, Hub }
+	};
+
+	bool CanReceiveFrom(ELiveLinkTopologyMode LHSMode, ELiveLinkTopologyMode RHSMode)
+	{
+		TArray<ELiveLinkTopologyMode> AllowedConnections;
+		CanReceiveFrom_Map.MultiFind(LHSMode, AllowedConnections);
+		return AllowedConnections.Contains(RHSMode);
+	}
+
+	bool CanTransmitTo(ELiveLinkTopologyMode LHSMode, ELiveLinkTopologyMode RHSMode)
+	{
+		TArray<ELiveLinkTopologyMode> AllowedConnections;
+		CanTransmitTo_Map.MultiFind(LHSMode, AllowedConnections);
+		return AllowedConnections.Contains(RHSMode);
+	}
+}
 
 
 // Address that we have had a connection request from
@@ -124,7 +170,7 @@ void FLiveLinkProvider::ValidateConnections()
 
 void FLiveLinkProvider::CloseConnection(FMessageAddress Address)
 {
-	TArray<FMessageAddress> RemovedConnections;
+	TSet<FMessageAddress> RemovedConnections;
 
 	{
 		FScopeLock Lock(&CriticalSection);
@@ -141,7 +187,7 @@ void FLiveLinkProvider::CloseConnection(FMessageAddress Address)
 
 	if (RemovedConnections.Num() > 0)
 	{
-		OnConnectionsClosed(RemovedConnections);
+		OnConnectionsClosed(RemovedConnections.Array());
 		OnConnectionStatusChanged.Broadcast();
 	}
 }
diff --git a/Engine/Source/Runtime/LiveLinkMessageBusFramework/Public/LiveLinkMessages.h b/Engine/Source/Runtime/LiveLinkMessageBusFramework/Public/LiveLinkMessages.h
--- a/Engine/Source/Runtime/LiveLinkMessageBusFramework/Public/LiveLinkMessages.h
+++ b/Engine/Source/Runtime/LiveLinkMessageBusFramework/Public/LiveLinkMessages.h
@@ -11,8 +11,29 @@ struct LIVELINKMESSAGEBUSFRAMEWORK_API FLiveLinkMessageAnnotation
 	static FName SubjectAnnotation;
 	static FName RoleAnnotation;
 	static FName OriginalSourceAnnotation;
+	/** Annotation to indicate what mode this instance is running in. */
+	static FName TopologyModeAnnotation;
 };
 
+UENUM()
+enum class ELiveLinkTopologyMode : uint8
+{
+	Hub   			UMETA(ToolTip = "LiveLinkHub instance can receive data from other hub instances in Spoke mode and transmit it to Unreal Editor instances."),
+	Spoke 			UMETA(ToolTip = "LiveLinkHub instance can only transmit data it to a LiveLinkHub instance in Hub mode."),
+	UnrealClient 	UMETA(ToolTip = "Unreal Editor client that can only transmit data it to a LiveLinkHub instance in Hub mode."),
+	External 		UMETA(ToolTip = "Client that can transmit to Hubs, Spokes and UnrealClients. (Default mode for message bus providers)")
+};
+
+
+namespace UE::LiveLink::Messaging
+{
+	/** Whether the left hand side mode can receive from the right hand side mode LiveLink Source. */
+	LIVELINKMESSAGEBUSFRAMEWORK_API bool CanReceiveFrom(ELiveLinkTopologyMode LHSMode, ELiveLinkTopologyMode RHSMode);
+
+	/** Whether the left hand side mode can transmit to the right hand side mode Client. */
+	LIVELINKMESSAGEBUSFRAMEWORK_API bool CanTransmitTo(ELiveLinkTopologyMode LHSMode, ELiveLinkTopologyMode RHSMode);
+}
+
 USTRUCT()
 struct FLiveLinkPingMessage
 {
@@ -124,4 +145,4 @@ struct UE_DEPRECATED(4.23, "FLiveLinkSubjectDataMessage is deprecated. Please us
 	UPROPERTY()
 	double Time = 0.0;
 };
-PRAGMA_ENABLE_DEPRECATION_WARNINGS
\ No newline at end of file
+PRAGMA_ENABLE_DEPRECATION_WARNINGS
diff --git a/Engine/Source/Runtime/LiveLinkMessageBusFramework/Public/LiveLinkProviderImpl.h b/Engine/Source/Runtime/LiveLinkMessageBusFramework/Public/LiveLinkProviderImpl.h
--- a/Engine/Source/Runtime/LiveLinkMessageBusFramework/Public/LiveLinkProviderImpl.h
+++ b/Engine/Source/Runtime/LiveLinkMessageBusFramework/Public/LiveLinkProviderImpl.h
@@ -100,7 +100,7 @@ protected:
 	TPair<UClass*, FLiveLinkStaticDataStruct*> GetLastSubjectStaticDataStruct(FName SubjectName);
 
 	template<typename MessageType>
-	void SendMessage(MessageType* Message)
+	void SendMessage(MessageType* Message, EMessageFlags Flags = EMessageFlags::None)
 	{
 		if (!Message)
 		{
@@ -111,30 +111,41 @@ protected:
 		GetConnectedAddresses(Addresses);
 		if (Addresses.Num() != 0)
 		{
-			MessageEndpoint->Send(Message, Addresses);
+			MessageEndpoint->Send(Message, Flags, {}, nullptr, Addresses, FTimespan::Zero(), FDateTime::MaxValue());
 		}
 	}
 
 	template<typename MessageType>
-	void SendMessage(MessageType* Message, const FMessageAddress& Address)
+	void SendMessage(MessageType* Message, const FMessageAddress& Address, EMessageFlags Flags = EMessageFlags::None)
 	{
 		if (!Message || !Address.IsValid())
 		{
 			return;
 		}
 
-		MessageEndpoint->Send(Message, Address);
+		MessageEndpoint->Send(Message, Flags, {}, nullptr, { Address }, FTimespan::Zero(), FDateTime::MaxValue());
 	}
 
 	template<typename MessageType>
-	void SendMessage(MessageType* Message, const TArray<FMessageAddress>& Addresses)
+	void SendMessage(MessageType* Message, const TArray<FMessageAddress>& Addresses, EMessageFlags Flags = EMessageFlags::None)
 	{
 		if (!Message || !Addresses.Num())
 		{
 			return;
 		}
 
-		MessageEndpoint->Send(Message, Addresses);
+		MessageEndpoint->Send(Message, Flags, {}, nullptr, Addresses, FTimespan::Zero(), FDateTime::MaxValue());
+	}
+
+	template<typename MessageType>
+	void Publish(MessageType* Message)
+	{
+		if (!Message)
+		{
+			return;
+		}
+
+		MessageEndpoint->Publish(Message);
 	}
 
 	template<typename MessageType>
@@ -162,6 +173,11 @@ protected:
 	// Get annotations to include on every message sent by this provider
 	virtual TMap<FName, FString> GetAnnotations() const
 	{
+#if WITH_ENGINE
+		// When running a game/editor, mark this provider as an unreal client.
+		FString TopologyModeValue = StaticEnum<ELiveLinkTopologyMode>()->GetNameStringByValue(static_cast<int64>(ELiveLinkTopologyMode::UnrealClient));
+		return { {FLiveLinkMessageAnnotation::TopologyModeAnnotation, MoveTemp(TopologyModeValue)} };
+#endif
 		return {};
 	}
 
