Fix CIS error caused by shadowed variable in LiveLinkHubMessaging

#rnx

diff --git a/Engine/Plugins/Animation/LiveLinkHub/Source/LiveLinkHubMessaging/Private/LiveLinkHubMessageBusSource.cpp b/Engine/Plugins/Animation/LiveLinkHub/Source/LiveLinkHubMessaging/Private/LiveLinkHubMessageBusSource.cpp
--- a/Engine/Plugins/Animation/LiveLinkHub/Source/LiveLinkHubMessaging/Private/LiveLinkHubMessageBusSource.cpp
+++ b/Engine/Plugins/Animation/LiveLinkHub/Source/LiveLinkHubMessaging/Private/LiveLinkHubMessageBusSource.cpp
@@ -75,7 +75,8 @@ bool FLiveLinkHubMessageBusSource::RequestSourceShutdown()
 	if (!bDisconnecting)
 	{
 		// Otherwise, inform our connected provider that it should remove this source.
-		SendMessage(FMessageEndpoint::MakeMessage<FLiveLinkHubDisconnectMessage>(GetSourceType().ToString(), GetSourceMachineName().ToString()));
+		FLiveLinkHubDisconnectMessage Message{ GetSourceType().ToString(), GetSourceMachineName().ToString() };
+		SendMessage(FMessageEndpoint::MakeMessage<FLiveLinkHubDisconnectMessage>(MoveTemp(Message)));
 	}
 
 	return FLiveLinkMessageBusSource::RequestSourceShutdown();
diff --git a/Engine/Plugins/Animation/LiveLinkHub/Source/LiveLinkHubMessaging/Private/LiveLinkHubMessagingModule.cpp b/Engine/Plugins/Animation/LiveLinkHub/Source/LiveLinkHubMessaging/Private/LiveLinkHubMessagingModule.cpp
--- a/Engine/Plugins/Animation/LiveLinkHub/Source/LiveLinkHubMessaging/Private/LiveLinkHubMessagingModule.cpp
+++ b/Engine/Plugins/Animation/LiveLinkHub/Source/LiveLinkHubMessaging/Private/LiveLinkHubMessagingModule.cpp
@@ -132,17 +132,18 @@ FLiveLinkHubInstanceId::FLiveLinkHubInstanceId(FStringView NamedId)
 FString FLiveLinkHubInstanceId::ToString() const
 {
 	auto Overload = Internal::TOverloaded{
-		[](FGuid Id) -> FString 
+		[](FGuid InGuid) -> FString 
 		{
 			TStringBuilder<20> LiveLinkHubName;
-			LiveLinkHubName << TEXT("Live Link Hub") << TEXT(" (") << *Id.ToString().Right(4).ToLower() << TEXT(")");
+			LiveLinkHubName << TEXT("Live Link Hub") << TEXT(" (") << *InGuid.ToString().Right(4).ToLower() << TEXT(")");
 			return LiveLinkHubName.ToString();
 		},
-		[](const FString& Id) -> FString 
+		[](const FString& InId) -> FString
 		{
-			return Id;
+			return InId;
 		}
 	};
+
 	return Visit(Overload, Id);
 }
 
