LiveLinkHub - Allow deselecting subject key in the custom time step

diff --git a/Engine/Plugins/Animation/LiveLinkHub/Source/LiveLinkHub/Private/UI/Widgets/SLiveLinkTimecode.cpp b/Engine/Plugins/Animation/LiveLinkHub/Source/LiveLinkHub/Private/UI/Widgets/SLiveLinkTimecode.cpp
--- a/Engine/Plugins/Animation/LiveLinkHub/Source/LiveLinkHub/Private/UI/Widgets/SLiveLinkTimecode.cpp
+++ b/Engine/Plugins/Animation/LiveLinkHub/Source/LiveLinkHub/Private/UI/Widgets/SLiveLinkTimecode.cpp
@@ -375,10 +375,12 @@ void SLiveLinkTimecode::SetCustomTimeStepSubject(const FName SubjectName)
 {
 	if (SubjectName == ActiveSubjectCustomTimeStep)
 	{
-		return;
+		ActiveSubjectCustomTimeStep = NAME_None;
+	}
+	else
+	{
+		ActiveSubjectCustomTimeStep = SubjectName;
 	}
-
-	ActiveSubjectCustomTimeStep = SubjectName;
 
 	const FLiveLinkHubCustomTimeStepSettings Settings = MakeCustomTimeStepSettings();
 
