LiveLinkHub - Add Hub and Spoke support to LiveLinkhub, which allows connecting multiple instances of LiveLink to each others.

Spokes can receive data from LiveLink providers (ie. Mobu, Vicon) and can only transmit it to LiveLinkHub instances in Hub mode.
Hubs can receive data from any source/Spokes and transmits this data to Unreal clients.

This behavior can be changed at runtime or set through a command line argument(-spoke or -hub)

diff --git a/Engine/Plugins/Animation/LiveLink/Source/LiveLinkEditor/Private/LiveLinkPanelController.cpp b/Engine/Plugins/Animation/LiveLink/Source/LiveLinkEditor/Private/LiveLinkPanelController.cpp
--- a/Engine/Plugins/Animation/LiveLink/Source/LiveLinkEditor/Private/LiveLinkPanelController.cpp
+++ b/Engine/Plugins/Animation/LiveLink/Source/LiveLinkEditor/Private/LiveLinkPanelController.cpp
@@ -36,13 +36,17 @@ FLiveLinkPanelController::FLiveLinkPanelController(TAttribute<bool> bInReadOnly)
 
 FLiveLinkPanelController::~FLiveLinkPanelController()
 {
-	if (Client)
+	if (IModularFeatures::Get().IsModularFeatureAvailable(ILiveLinkClient::ModularFeatureName))
 	{
-		Client->OnLiveLinkSourcesChanged().Remove(OnSourcesChangedHandle);
-		OnSourcesChangedHandle.Reset();
+		ILiveLinkClient* LiveLinkClient = (FLiveLinkClient*)&IModularFeatures::Get().GetModularFeature<ILiveLinkClient>(ILiveLinkClient::ModularFeatureName);
+		if (LiveLinkClient)
+		{
+			LiveLinkClient->OnLiveLinkSourcesChanged().Remove(OnSourcesChangedHandle);
+			OnSourcesChangedHandle.Reset();
 
-		Client->OnLiveLinkSubjectsChanged().Remove(OnSubjectsChangedHandle);
-		OnSubjectsChangedHandle.Reset();
+			LiveLinkClient->OnLiveLinkSubjectsChanged().Remove(OnSubjectsChangedHandle);
+			OnSubjectsChangedHandle.Reset();
+		}
 	}
 }
 
diff --git a/Engine/Plugins/Animation/LiveLinkHub/Source/LiveLinkHub/Private/Clients/LiveLinkHubProvider.cpp b/Engine/Plugins/Animation/LiveLinkHub/Source/LiveLinkHub/Private/Clients/LiveLinkHubProvider.cpp
--- a/Engine/Plugins/Animation/LiveLinkHub/Source/LiveLinkHub/Private/Clients/LiveLinkHubProvider.cpp
+++ b/Engine/Plugins/Animation/LiveLinkHub/Source/LiveLinkHub/Private/Clients/LiveLinkHubProvider.cpp
@@ -49,7 +49,6 @@ namespace LiveLinkHubProviderUtils
 			}
 		}
 
-
 		ensureMsgf(false, TEXT("Feature %s is unavailable"), *INetworkMessagingExtension::ModularFeatureName.ToString());
 		return nullptr;
 	}
@@ -85,6 +84,7 @@ FLiveLinkHubProvider::FLiveLinkHubProvider(const TSharedRef<ILiveLinkHubSessionM
 	FMessageEndpointBuilder EndpointBuilder = FMessageEndpoint::Builder(*GetProviderName());
 	EndpointBuilder.WithHandler(MakeHandler<FLiveLinkClientInfoMessage>(&FLiveLinkHubProvider::HandleClientInfoMessage));
 	EndpointBuilder.WithHandler(MakeHandler<FLiveLinkHubConnectMessage>(&FLiveLinkHubProvider::HandleHubConnectMessage));
+	EndpointBuilder.WithHandler(MakeHandler<FLiveLinkHubDisconnectMessage>(&FLiveLinkHubProvider::HandleHubDisconnectMessage));
 
 	CreateMessageEndpoint(EndpointBuilder);
 
@@ -138,6 +138,66 @@ void FLiveLinkHubProvider::ResetCustomTimeStepSettings(const FLiveLinkHubClientI
 	SendCustomTimeStepSettings(ResetCustomTimeStepSettings, ClientId);
 }
 
+void FLiveLinkHubProvider::DisconnectAll()
+{
+	UE_LOG(LogLiveLinkHub, Verbose, TEXT("Provider: Sending DisconnectAll message (%s)"), *GetEndpointAddress().ToString());
+
+	FLiveLinkHubDisconnectMessage DisconnectMessage;
+	DisconnectMessage.ProviderName = GetProviderName();
+	DisconnectMessage.MachineName = GetMachineName();
+
+	// todo: Clear address cache to prevent getting connection closed notifications
+	//AddressToIdCache.Reset()
+
+	SendMessage(FMessageEndpoint::MakeMessage<FLiveLinkHubDisconnectMessage>(MoveTemp(DisconnectMessage)), EMessageFlags::Reliable);
+	
+	TArray<FMessageAddress> AllAddresses;
+	GetConnectedAddresses(AllAddresses);
+	for (const FMessageAddress& Address : AllAddresses)
+	{
+		CloseConnection(Address);
+	}
+}
+
+void FLiveLinkHubProvider::DisconnectClient(const FLiveLinkHubClientId& Client)
+{
+	FLiveLinkHubDisconnectMessage DisconnectMessage;
+	DisconnectMessage.ProviderName = GetProviderName();
+	DisconnectMessage.MachineName = GetMachineName();
+
+	TArray<FMessageAddress> AllAddresses;
+	GetConnectedAddresses(AllAddresses);
+
+	FMessageAddress* TargetAddress = AllAddresses.FindByPredicate(
+		[this, Client](const FMessageAddress& Address)
+		{
+			return AddressToIdCache.FindChecked(Address) == Client;
+		});
+
+	// Todo: We may want to remove the client info and tracked address.
+	if (TargetAddress)
+	{
+		UE_LOG(LogLiveLinkHub, Verbose, TEXT("Provider: Sending Disconnect message from %s to %s"), *GetEndpointAddress().ToString(), *TargetAddress->ToString());
+		SendMessage(FMessageEndpoint::MakeMessage<FLiveLinkHubDisconnectMessage>(MoveTemp(DisconnectMessage)), *TargetAddress, EMessageFlags::Reliable);
+		CloseConnection(*TargetAddress);
+	}
+}
+
+void FLiveLinkHubProvider::PostUpdateTopologyMode(ELiveLinkTopologyMode TopologyMode)
+{
+	// Hub -> Spoke: Disconnect from all clients, then connect to discovered hubs
+	// Spokes -> Hub: Disconnect from everything all clients, then connect to discovered UE clients
+
+	if (TopologyMode == ELiveLinkTopologyMode::Hub)
+	{
+		ConnectToAllUEClients();
+	}
+	else if (TopologyMode == ELiveLinkTopologyMode::Spoke)
+	{
+		ConnectToAllHubClients();
+	}
+}
+
 void FLiveLinkHubProvider::SendTimecodeSettings(const FLiveLinkHubTimecodeSettings& InSettings, const FLiveLinkHubClientId& ClientId)
 {
 	if (ClientId.IsValid())
@@ -147,12 +207,19 @@ void FLiveLinkHubProvider::SendTimecodeSettings(const FLiveLinkHubTimecodeSettin
 
 		FMessageAddress* TargetAddress = AllAddresses.FindByPredicate([this, ClientId](const FMessageAddress& Address)
 		{
-			return AddressToIdCache.FindChecked(Address) == ClientId;
+			if (FLiveLinkHubClientId* FoundId = AddressToIdCache.Find(Address))
+			{
+				return *FoundId == ClientId;
+			}
+
+			UE_LOG(LogLiveLinkHub, Warning, TEXT("Could not find client id for address %s."), *Address.ToString());
+			return false;
+
 		});
 
 		if (TargetAddress)
 		{
-			SendMessage(FMessageEndpoint::MakeMessage<FLiveLinkHubTimecodeSettings>(InSettings));
+			SendMessage(FMessageEndpoint::MakeMessage<FLiveLinkHubTimecodeSettings>(InSettings), *TargetAddress, EMessageFlags::Reliable);
 		}
 	}
 	else
@@ -176,13 +243,13 @@ void FLiveLinkHubProvider::SendCustomTimeStepSettings(const FLiveLinkHubCustomTi
 
 		if (TargetAddress)
 		{
-			SendMessage(FMessageEndpoint::MakeMessage<FLiveLinkHubCustomTimeStepSettings>(InSettings));
+			SendMessage(FMessageEndpoint::MakeMessage<FLiveLinkHubCustomTimeStepSettings>(InSettings), *TargetAddress, EMessageFlags::Reliable);
 		}
 	}
 	else
 	{
 		// Invalid ID means we're broadcasting to all clients.
-		SendMessageToEnabledClients(FMessageEndpoint::MakeMessage<FLiveLinkHubCustomTimeStepSettings>(InSettings));
+		SendMessageToEnabledClients(FMessageEndpoint::MakeMessage<FLiveLinkHubCustomTimeStepSettings>(InSettings), EMessageFlags::Reliable);
 	}
 }
 
@@ -233,6 +300,17 @@ TOptional<FLiveLinkHubUEClientInfo> FLiveLinkHubProvider::GetClientInfo(FLiveLin
 
 void FLiveLinkHubProvider::HandleHubConnectMessage(const FLiveLinkHubConnectMessage& Message, const TSharedRef<IMessageContext, ESPMode::ThreadSafe>& Context)
 {
+	UE_LOG(LogLiveLinkHub, Verbose, TEXT("Provider: Received connect message from %s"), *Context->GetSender().ToString());
+
+	const ELiveLinkTopologyMode Mode = FLiveLinkHub::Get()->GetTopologyMode();
+	if (!UE::LiveLink::Messaging::CanTransmitTo(Mode, Message.ClientInfo.TopologyMode))
+	{
+		UE_LOG(LogLiveLinkHub, Verbose, TEXT("Provider: Denying connection from %s since its mode is incompatible with this provider's."), *Context->GetSender().ToString());
+		SendMessage(FMessageEndpoint::MakeMessage<FLiveLinkHubDisconnectMessage>(GetProviderName(), GetMachineName()), Context->GetSender(), EMessageFlags::Reliable);
+		CloseConnection(Context->GetSender());
+		return;
+	}
+
 	FLiveLinkConnectMessage ConnectMessage;
 	ConnectMessage.LiveLinkVersion = Message.ClientInfo.LiveLinkVersion;
 	FLiveLinkProvider::HandleConnectMessage(ConnectMessage, Context);
@@ -343,6 +421,8 @@ void FLiveLinkHubProvider::HandleHubConnectMessage(const FLiveLinkHubConnectMess
 	
 void FLiveLinkHubProvider::HandleClientInfoMessage(const FLiveLinkClientInfoMessage& Message, const TSharedRef<IMessageContext, ESPMode::ThreadSafe>& Context)
 {
+	UE_LOG(LogLiveLinkHub, Verbose, TEXT("Provider: Received ClientInfo message from %s"), *Context->GetSender().ToString());
+
 	FMessageAddress Address = Context->GetSender();
 
 	FLiveLinkHubClientId ClientId = AddressToIdCache.FindRef(Address);
@@ -371,6 +451,22 @@ void FLiveLinkHubProvider::HandleClientInfoMessage(const FLiveLinkClientInfoMess
 	}
 }
 
+void FLiveLinkHubProvider::HandleHubDisconnectMessage(const FLiveLinkHubDisconnectMessage& Message, const TSharedRef<IMessageContext, ESPMode::ThreadSafe>& Context)
+{
+	UE_LOG(LogLiveLinkHub, Verbose, TEXT("Provider: Received disconnect message from %s"), *Context->GetSender().ToString());
+
+	// Received a disconnect message from the livelinkhub source (it was probably deleted), so let's remove this client from the session
+	const FMessageAddress Address = Context->GetSender();
+
+	// Disabling it, so reset its timecode and custom time step.
+	FLiveLinkHubClientId ClientId = AddressToIdCache.FindRef(Address);
+	if (TSharedPtr<ILiveLinkHubSessionManager> Manager = SessionManager.Pin())
+	{
+		// The session will handle resetting the timecode settings and removing it from the cache.
+		Manager->GetCurrentSession()->RemoveClient(ClientId);
+	}
+}
+
 bool FLiveLinkHubProvider::ShouldTransmitToClient_AnyThread(FMessageAddress Address, TFunctionRef<bool(const FLiveLinkHubUEClientInfo* ClientInfoPtr)> AdditionalFilter) const
 {
 	if (!Address.IsValid())
@@ -409,7 +505,35 @@ bool FLiveLinkHubProvider::ShouldTransmitToClient_AnyThread(FMessageAddress Addr
 	return true;
 }
 
+void FLiveLinkHubProvider::ConnectToAllHubClients()
+{
+	Publish(FMessageEndpoint::MakeMessage<FLiveLinkHubDiscoveryMessage>(GetProviderName(), ELiveLinkTopologyMode::Spoke, FLiveLinkHub::Get()->GetId()));
+}
+
+void FLiveLinkHubProvider::ConnectToAllUEClients()
+{
+	Publish(FMessageEndpoint::MakeMessage<FLiveLinkHubDiscoveryMessage>(GetProviderName(), ELiveLinkTopologyMode::Hub, FLiveLinkHub::Get()->GetId()));
+}
+
 void FLiveLinkHubProvider::OnConnectionsClosed(const TArray<FMessageAddress>& ClosedAddresses)
+{
+	CloseConnections(ClosedAddresses);
+
+	for (FMessageAddress TrackedAddress : ClosedAddresses)
+	{
+		FWriteScopeLock Locker(ClientsMapLock);
+		if (FLiveLinkHubClientId* ClientId = AddressToIdCache.Find(TrackedAddress))
+		{
+			// Removing this might have implications for restoring sessions.
+			// We could instead remove this when the connection is forcibly closed.
+			ClientsMap.Remove(*ClientId);
+		}
+
+		AddressToIdCache.Remove(TrackedAddress);
+	}
+}
+
+void FLiveLinkHubProvider::CloseConnections(const TArray<FMessageAddress>& ClosedAddresses)
 {
 	// List of OldId -> NewId
 	TArray<FLiveLinkHubClientId> Notifications;
@@ -432,12 +556,6 @@ void FLiveLinkHubProvider::OnConnectionsClosed(const TArray<FMessageAddress>& Cl
 	{
 		OnClientEventDelegate.Broadcast(Client, EClientEventType::Disconnected);
 	}
-
-	for (FMessageAddress TrackedAddress : ClosedAddresses)
-	{
-		FWriteScopeLock Locker(ClientsMapLock);
-		AddressToIdCache.Remove(TrackedAddress);
-	}
 }
 
 TArray<FLiveLinkHubClientId> FLiveLinkHubProvider::GetSessionClients() const
@@ -463,8 +581,13 @@ TMap<FName, FString> FLiveLinkHubProvider::GetAnnotations() const
 	const FString AutoConnectValue = StaticEnum<ELiveLinkHubAutoConnectMode>()->GetNameStringByValue(static_cast<int64>(GetDefault<ULiveLinkHubSettings>()->AutoConnectClients));
 	AnnotationsCopy.Add(FLiveLinkHubMessageAnnotation::AutoConnectModeAnnotation, AutoConnectValue);
 
+	// In case this is handled by the discovery manager (which doesn't directly handle LiveLinkHub messages)
+	AnnotationsCopy.Add(FLiveLinkHubMessageAnnotation::IdAnnotation, FLiveLinkHub::Get()->GetId().ToString());
+
 	TMap<FName, FString> BaseAnnotations = FLiveLinkProvider::GetAnnotations();
 	BaseAnnotations.Append(AnnotationsCopy);
+	// In case this is handled by the discovery manager (which doesn't directly handle LiveLinkHub messages)
+	AnnotationsCopy.Add(FLiveLinkHubMessageAnnotation::IdAnnotation, FLiveLinkHub::Get()->GetId().ToString());
 
 	return BaseAnnotations;
 }
@@ -499,7 +622,14 @@ FText FLiveLinkHubProvider::GetClientDisplayName(FLiveLinkHubClientId InAddress)
 
 	if (const FLiveLinkHubUEClientInfo* ClientInfoPtr = ClientsMap.Find(InAddress))
 	{
-		DisplayName = FText::FromString(FString::Format(TEXT("{0} ({1})"), {*ClientInfoPtr->Hostname, *ClientInfoPtr->CurrentLevel}));
+		if (ClientInfoPtr->TopologyMode == ELiveLinkTopologyMode::Hub)
+		{
+			DisplayName = FText::FromString(ClientInfoPtr->LiveLinkInstanceName);
+		}
+		else
+		{
+			DisplayName = FText::FromString(FString::Format(TEXT("{0} ({1})"), {*ClientInfoPtr->Hostname, *ClientInfoPtr->CurrentLevel }));
+		}
 	}
 	else
 	{
diff --git a/Engine/Plugins/Animation/LiveLinkHub/Source/LiveLinkHub/Private/Clients/LiveLinkHubProvider.h b/Engine/Plugins/Animation/LiveLinkHub/Source/LiveLinkHub/Private/Clients/LiveLinkHubProvider.h
--- a/Engine/Plugins/Animation/LiveLinkHub/Source/LiveLinkHub/Private/Clients/LiveLinkHubProvider.h
+++ b/Engine/Plugins/Animation/LiveLinkHub/Source/LiveLinkHub/Private/Clients/LiveLinkHubProvider.h
@@ -76,6 +76,19 @@ public:
 	 */
 	void ResetCustomTimeStepSettings(const FLiveLinkHubClientId& ClientId = FLiveLinkHubClientId{});
 
+	/**
+	 * Send a disconnect message to all connected clients.
+	 */
+	void DisconnectAll();
+
+	/**
+	 * Send a disconnect message to a single client.
+	 */
+	void DisconnectClient(const FLiveLinkHubClientId& ClientId);
+
+	/** Called when the topology mode changed for this app. */
+	void PostUpdateTopologyMode(ELiveLinkTopologyMode TopologyMode);
+
 private:
 	/** Handle a connection message resulting from a livelink hub message bus source connecting to this provider. */
 	void HandleHubConnectMessage(const FLiveLinkHubConnectMessage& Message, const TSharedRef<IMessageContext, ESPMode::ThreadSafe>& Context);
@@ -83,6 +96,9 @@ private:
 	/** Handle a client info message being received. Happens when new information about a client is received (ie. Client has changed map) */
 	void HandleClientInfoMessage(const FLiveLinkClientInfoMessage& Message, const TSharedRef<IMessageContext, ESPMode::ThreadSafe>& Context);
 
+	/** Handle a client (UE or Hub) sending a disconnect request. */
+	void HandleHubDisconnectMessage(const FLiveLinkHubDisconnectMessage& Message, const TSharedRef<IMessageContext, ESPMode::ThreadSafe>& Context);
+
 	/** Send timecode settings to connected UE Clients. */
 	void SendTimecodeSettings(const FLiveLinkHubTimecodeSettings& InSettings, const FLiveLinkHubClientId& ClientId);
 
@@ -91,7 +107,7 @@ private:
 
 	/** Send a message to clients that are connected and enabled through the hub clients list. */
 	template<typename MessageType>
-	void SendMessageToEnabledClients(MessageType* Message)
+	void SendMessageToEnabledClients(MessageType* Message, EMessageFlags Flags = EMessageFlags::None)
 	{
 		TArray<FMessageAddress> AllAddresses;
 		GetConnectedAddresses(AllAddresses);
@@ -101,7 +117,7 @@ private:
 			return ShouldTransmitToClient_AnyThread(Address);
 		});
 
-		SendMessage(Message, EnabledAddresses);
+		SendMessage(Message, EnabledAddresses, Flags);
 	}
 
 	/**
@@ -110,6 +126,14 @@ private:
 	 **/
 	bool ShouldTransmitToClient_AnyThread(FMessageAddress Address, TFunctionRef<bool(const FLiveLinkHubUEClientInfo* ClientInfoPtr)> AdditionalFilter = [](const FLiveLinkHubUEClientInfo* ClientInfoPtr){ return true; }) const;
 
+	/** Broadcast a message telling all UE clients to connect to this instance. */
+	void ConnectToAllHubClients();
+
+	/** Broadcast a message telling all Hub clients to connect to this instance. */
+	void ConnectToAllUEClients();
+
+	void CloseConnections(const TArray<FMessageAddress>& ClosedAddresses);
+
 protected:
 	//~ Begin ILiveLinkHubClientsModel interface
 	virtual void OnConnectionsClosed(const TArray<FMessageAddress>& ClosedAddresses) override;
diff --git a/Engine/Plugins/Animation/LiveLinkHub/Source/LiveLinkHub/Private/Clients/LiveLinkHubUEClientInfo.h b/Engine/Plugins/Animation/LiveLinkHub/Source/LiveLinkHub/Private/Clients/LiveLinkHubUEClientInfo.h
--- a/Engine/Plugins/Animation/LiveLinkHub/Source/LiveLinkHub/Private/Clients/LiveLinkHubUEClientInfo.h
+++ b/Engine/Plugins/Animation/LiveLinkHub/Source/LiveLinkHub/Private/Clients/LiveLinkHubUEClientInfo.h
@@ -63,6 +63,8 @@ struct FLiveLinkHubUEClientInfo
 		, Hostname(InClientInfo.Hostname)
 		, ProjectName(InClientInfo.ProjectName)
 		, CurrentLevel(InClientInfo.CurrentLevel)
+		, LiveLinkInstanceName(InClientInfo.LiveLinkInstanceName)
+		, TopologyMode(InClientInfo.TopologyMode)
 	{
 	}
 
@@ -102,6 +104,14 @@ struct FLiveLinkHubUEClientInfo
 	/** Name of the current level opened. */
 	UPROPERTY(VisibleAnywhere, Category = "Client Details")
 	FString CurrentLevel;
+
+	/** If this is representing a LiveLinkHub instance in Hub mode, this holds the LiveLink provider name, otherwise it's empty. */
+	UPROPERTY()
+	FString LiveLinkInstanceName;
+
+	/** Whether the client is a hub or an unreal instance. */
+	UPROPERTY()
+	ELiveLinkTopologyMode TopologyMode = ELiveLinkTopologyMode::UnrealClient;
 	
 	/** Subjects that should not be transmitted to this client. */
 	UPROPERTY()
diff --git a/Engine/Plugins/Animation/LiveLinkHub/Source/LiveLinkHub/Private/LiveLinkHub.cpp b/Engine/Plugins/Animation/LiveLinkHub/Source/LiveLinkHub/Private/LiveLinkHub.cpp
--- a/Engine/Plugins/Animation/LiveLinkHub/Source/LiveLinkHub/Private/LiveLinkHub.cpp
+++ b/Engine/Plugins/Animation/LiveLinkHub/Source/LiveLinkHub/Private/LiveLinkHub.cpp
@@ -11,16 +11,19 @@
 #include "Framework/Application/SlateApplication.h"
 #include "HAL/PlatformProcess.h"
 #include "ISettingsModule.h"
+#include "ILiveLinkHubMessagingModule.h"
 #include "LiveLinkEditorSettings.h"
 #include "LiveLinkHubClient.h"
 #include "LiveLinkHubCommands.h"
 #include "LiveLinkHubCreatorAppMode.h"
 #include "LiveLinkHubModule.h"
+#include "LiveLinkHubMessageBusSourceSettings.h"
 #include "LiveLinkHubSubjectSettings.h"
 #include "LiveLinkHubTicker.h"
 #include "LiveLinkProviderImpl.h"
 #include "LiveLinkSettings.h"
 #include "Misc/App.h"
+#include "Misc/CommandLine.h"
 #include "Misc/ConfigCacheIni.h"
 #include "Recording/LiveLinkHubPlaybackController.h"
 #include "Recording/LiveLinkHubRecordingController.h"
@@ -63,18 +66,15 @@ namespace LiveLinkHubInternalUtils
 	}
 
 	/** Generate the name of the livelinkhub provider. */
-	FString GetProviderName()
+	FString GetProviderName(const FLiveLinkHubInstanceId& Id)
 	{
-		FString ProviderName = TEXT("Live Link Hub");
-
 		uint32 InstanceCount = GetInstanceCount();
 		if (InstanceCount > 1)
 		{
-			// Append a unique ID if more than one instance is running.
-			ProviderName += FString::Printf(TEXT(" (%s)"), *FGuid::NewGuid().ToString().Right(4).ToLower());
+			return Id.ToString();
 		}
 
-		return ProviderName;
+		return TEXT("Live Link Hub");;
 	}
 }
 
@@ -109,7 +109,7 @@ void FLiveLinkHub::Initialize(FLiveLinkHubTicker& Ticker)
 
 	SessionManager = MakeShared<FLiveLinkHubSessionManager>();
 
-	LiveLinkProvider = MakeShared<FLiveLinkHubProvider>(SessionManager.ToSharedRef(), LiveLinkHubInternalUtils::GetProviderName());
+	LiveLinkProvider = MakeShared<FLiveLinkHubProvider>(SessionManager.ToSharedRef(), LiveLinkHubInternalUtils::GetProviderName(InstanceId));
 
 	FModuleManager::Get().LoadModule("Settings");
 	FModuleManager::Get().LoadModule("StatusBar");
@@ -132,6 +132,19 @@ void FLiveLinkHub::Initialize(FLiveLinkHubTicker& Ticker)
 		ClientsController = MakeShared<FLiveLinkHubClientsController>(LiveLinkProvider.ToSharedRef());
 		SubjectController = MakeShared<FLiveLinkHubSubjectController>();
 	}
+	
+	if (FParse::Param(FCommandLine::Get(), TEXT("Hub")))
+	{
+		OverrideTopologyMode = ELiveLinkTopologyMode::Hub;
+	}
+	else if (FParse::Param(FCommandLine::Get(), TEXT("Spoke")))
+	{
+		OverrideTopologyMode = ELiveLinkTopologyMode::Spoke;
+	}
+
+	// We load it now to ensure we set the connection mode as early as possible to avoid it discovering instances by mistake.
+	ILiveLinkHubMessagingModule* MessagingModule = static_cast<ILiveLinkHubMessagingModule*>(FModuleManager::Get().LoadModule("LiveLinkHubMessaging"));
+	MessagingModule->SetHostTopologyMode(GetTopologyMode());
 
 	FCoreDelegates::OnPostEngineInit.AddRaw(this, &FLiveLinkHub::InitializePostEngineInit);
 }
@@ -142,6 +155,11 @@ void FLiveLinkHub::InitializePostEngineInit()
 
 	TabManager = FGlobalTabmanager::Get();
 
+	ILiveLinkHubMessagingModule& MessagingModule = FModuleManager::Get().GetModuleChecked<ILiveLinkHubMessagingModule>("LiveLinkHubMessaging");
+
+	// Update the mode in our connection manager as well.
+	MessagingModule.SetInstanceId(InstanceId);
+
 	AddLiveLinkHubApplicationMode("CreatorMode", MakeShared<FLiveLinkHubCreatorAppMode>(SharedThis<FWorkflowCentricApplication>(this)));
 
 	for (ILiveLinkHubApplicationModeFactory* Factory : IModularFeatures::Get().GetModularFeatureImplementations<ILiveLinkHubApplicationModeFactory>(ILiveLinkHubApplicationModeFactory::ModularFeatureName))
@@ -170,8 +188,6 @@ void FLiveLinkHub::InitializePostEngineInit()
 
 FLiveLinkHub::~FLiveLinkHub()
 {
-	RemoveAllApplicationModes();
-
 	IModularFeatures::Get().UnregisterModularFeature(IConsoleCommandExecutor::ModularFeatureName(), CommandExecutor.Get());
 	CommandExecutor.Reset();
 
@@ -236,6 +252,82 @@ TSharedPtr<FLiveLinkHubPlaybackController> FLiveLinkHub::GetPlaybackController()
 	return PlaybackController;
 }
 
+ELiveLinkTopologyMode FLiveLinkHub::GetTopologyMode() const
+{
+	if (OverrideTopologyMode)
+	{
+		return *OverrideTopologyMode;
+	}
+
+	if (SessionManager)
+	{
+		if (TSharedPtr<ILiveLinkHubSession> CurrentSession = SessionManager->GetCurrentSession())
+		{
+			return CurrentSession->GetTopologyMode();
+		}
+	}
+
+	return ELiveLinkTopologyMode::Hub;
+}
+
+void FLiveLinkHub::SetTopologyMode(ELiveLinkTopologyMode Mode)
+{
+	checkf(!OverrideTopologyMode, TEXT("Can't set topology mode at runtime if it was set through the command line."));
+
+	if (Mode != GetTopologyMode())
+	{
+		ILiveLinkHubMessagingModule& MessagingModule = FModuleManager::Get().GetModuleChecked<ILiveLinkHubMessagingModule>("LiveLinkHubMessaging");
+
+		// Update the mode in our connection manager as well.
+		MessagingModule.SetHostTopologyMode(Mode);
+
+		if (SessionManager)
+		{
+			if (TSharedPtr<ILiveLinkHubSession> CurrentSession = SessionManager->GetCurrentSession())
+			{
+				CurrentSession->SetTopologyMode(Mode);
+			}
+		}
+
+		if (Mode == ELiveLinkTopologyMode::Spoke)
+		{
+			// Clear out the list of LLH sources if we're switching from a hub to a spoke.
+			constexpr bool bEvenIfPendingKill = true;
+			for (const FGuid& Source : LiveLinkHubClient->GetSources(bEvenIfPendingKill))
+			{
+				if (ULiveLinkSourceSettings* Settings = LiveLinkHubClient->GetSourceSettings(Source))
+				{
+					if (Settings->IsA<ULiveLinkHubMessageBusSourceSettings>())
+					{
+						LiveLinkHubClient->RemoveSource(Source);
+					}
+				}
+			}
+		}
+
+		TopologyModeChangedDelegate.Broadcast(Mode);
+	}
+}
+
+void FLiveLinkHub::ToggleTopologyMode()
+{
+	ELiveLinkTopologyMode CurrentMode = GetTopologyMode();
+	if (CurrentMode == ELiveLinkTopologyMode::Hub)
+	{
+		SetTopologyMode(ELiveLinkTopologyMode::Spoke);
+	}
+	else
+	{
+		SetTopologyMode(ELiveLinkTopologyMode::Hub);
+	}
+}
+
+bool FLiveLinkHub::CanSetTopologyMode() const
+{
+	// Can't modify topology if it was overridden through the command line.
+	return !OverrideTopologyMode.IsSet();
+}
+
 void FLiveLinkHub::MapToolkitCommands()
 {
 	const FLiveLinkHubCommands& Commands = FLiveLinkHubCommands::Get();
@@ -250,6 +342,18 @@ void FLiveLinkHub::MapToolkitCommands()
 	FLiveLinkHubApplication::MapToolkitCommands();
 }
 
+void FLiveLinkHub::OnClose()
+{
+	// Make sure to send a disconnect message to every UE Client to ensure that they remove their LiveLinkHub source.
+	if (TSharedPtr<ILiveLinkHubSession> CurrentSession = SessionManager->GetCurrentSession())
+	{
+		// Note: This has to be done before the LiveLink Provider is destroyed.
+		CurrentSession->RemoveAllClients();
+	}
+
+	FWorkflowCentricApplication::OnClose();
+}
+
 void FLiveLinkHub::OnStaticDataReceived_AnyThread(const FLiveLinkSubjectKey& InSubjectKey, TSubclassOf<ULiveLinkRole> InRole, const FLiveLinkStaticDataStruct& InStaticDataStruct) const
 {
 	if (RecordingController->IsRecording())
diff --git a/Engine/Plugins/Animation/LiveLinkHub/Source/LiveLinkHub/Private/LiveLinkHub.h b/Engine/Plugins/Animation/LiveLinkHub/Source/LiveLinkHub/Private/LiveLinkHub.h
--- a/Engine/Plugins/Animation/LiveLinkHub/Source/LiveLinkHub/Private/LiveLinkHub.h
+++ b/Engine/Plugins/Animation/LiveLinkHub/Source/LiveLinkHub/Private/LiveLinkHub.h
@@ -4,7 +4,10 @@
 
 #include "LiveLinkHubApplication.h"
 
+#include "ILiveLinkHubMessagingModule.h"
 #include "LiveLinkTypes.h"
+#include "LiveLinkHubMessages.h"
+#include "Misc/Optional.h"
 #include "Templates/SubclassOf.h"
 #include "Templates/UniquePtr.h"
 #include "Toolkits/FConsoleCommandExecutor.h"
@@ -22,6 +25,8 @@ class ILiveLinkHubSessionManager;
 class SWindow;
 class ULiveLinkRole;
 
+DECLARE_TS_MULTICAST_DELEGATE_OneParam(FOnTopologyModeChanged, ELiveLinkTopologyMode /*NewMode*/);
+
 /**
  * Main interface for the live link hub.
  */
@@ -60,6 +65,7 @@ public:
 	void Initialize(class FLiveLinkHubTicker& Ticker);
 	/** Tick the hub. */
 	void Tick();
+	const FLiveLinkHubInstanceId& GetId() const { return InstanceId; }
 	/** Get the root window that hosts the hub's slate application. */
 	TSharedRef<SWindow> GetRootWindow() const;
 	/** Get the livelink provider used to rebroadcast livelink data to connected UE clients. */
@@ -76,9 +82,24 @@ public:
 	TSharedPtr<FUICommandList> GetCommandList() const { return GetToolkitCommands(); }
 	/** Get the session manager. */
 	TSharedPtr<ILiveLinkHubSessionManager> GetSessionManager() const;
+	/** 
+	 * Get whether the hub should act as a hub or as a spoke. 
+	 * Hubs can receive data from spokes, and transmit it to UE clients.
+	 * Spokes can only transmit data to Hubs.
+	 */
+	ELiveLinkTopologyMode GetTopologyMode() const;
+	/** Set the topology mode for this instance. */
+	void SetTopologyMode(ELiveLinkTopologyMode Mode);
+	/** Toggle the topology mode between Hub and Spoke. */
+	void ToggleTopologyMode();
+	/** Returns whether the topology mode can change for this instance. */
+	bool CanSetTopologyMode() const;
+	/** Delegate called when the topology mode changed. */
+	FOnTopologyModeChanged& OnTopologyModeChanged() { return TopologyModeChangedDelegate; }
 
 	//~ FAssetEditorToolkit interface
 	virtual void MapToolkitCommands() override;
+	virtual void OnClose() override;
 
 private:
 	/** Complete the initialization after core engine systems have been initialized. */
@@ -136,6 +157,12 @@ private:
 	TUniquePtr<FConsoleCommandExecutor> CommandExecutor;
 	/** The last opened config path. */
 	FString LastConfigPath;
+	/** LiveLinkHub instance Id, used to disambiguate different instances. */
+	FLiveLinkHubInstanceId InstanceId = FLiveLinkHubInstanceId{ FGuid::NewGuid() };
+	/** Override topology mode set through the command line. Cannot be changed at runtime. */
+	TOptional<ELiveLinkTopologyMode> OverrideTopologyMode;
+	/** Delegate called when the topology mode changed. */
+	FOnTopologyModeChanged TopologyModeChangedDelegate;
 
 	friend class FLiveLinkHubModule;
 };
diff --git a/Engine/Plugins/Animation/LiveLinkHub/Source/LiveLinkHub/Private/LiveLinkHubModule.cpp b/Engine/Plugins/Animation/LiveLinkHub/Source/LiveLinkHub/Private/LiveLinkHubModule.cpp
--- a/Engine/Plugins/Animation/LiveLinkHub/Source/LiveLinkHub/Private/LiveLinkHubModule.cpp
+++ b/Engine/Plugins/Animation/LiveLinkHub/Source/LiveLinkHub/Private/LiveLinkHubModule.cpp
@@ -82,7 +82,8 @@ void FLiveLinkHubModule::ShutdownModule()
 	{
 		Ticker.Stop();
 		Ticker.Exit();
-
+		// Application modes keep a shared reference to LiveLinkHub so we have to clear them before resetting the livelinkhub shared ptr.
+		LiveLinkHub->RemoveAllApplicationModes();
 		LiveLinkHub.Reset();
 	}
 }
diff --git a/Engine/Plugins/Animation/LiveLinkHub/Source/LiveLinkHub/Private/Session/LiveLinkHubSession.h b/Engine/Plugins/Animation/LiveLinkHub/Source/LiveLinkHub/Private/Session/LiveLinkHubSession.h
--- a/Engine/Plugins/Animation/LiveLinkHub/Source/LiveLinkHub/Private/Session/LiveLinkHubSession.h
+++ b/Engine/Plugins/Animation/LiveLinkHub/Source/LiveLinkHub/Private/Session/LiveLinkHubSession.h
@@ -38,6 +38,9 @@ public:
 	/** Remove a client from this session. Note: Must be called from game thread. */
 	virtual void RemoveClient(const FLiveLinkHubClientId& Client) = 0;
 
+	/** Remove all clients from this session. Note: Must be called from game thread. */
+	virtual void RemoveAllClients() = 0;
+
 	/** Returns whether a client is in this session. */
 	virtual bool IsClientInSession(const FLiveLinkHubClientId& Client) = 0;
 
@@ -67,6 +70,12 @@ public:
 
 	/** Set the CustomTimeStep settings for the current session. */
 	virtual void SetCustomTimeStepSettings(const FLiveLinkHubCustomTimeStepSettings& CustomTimeStepSettings) const = 0;
+
+	/** Get the topology mode for this instance. */
+	virtual ELiveLinkTopologyMode GetTopologyMode() const = 0;
+
+	/** Set the topology mode for this instance. */
+	virtual void SetTopologyMode(ELiveLinkTopologyMode Mode) = 0;
 };
 
 class FLiveLinkHubSession : public ILiveLinkHubSession, public TSharedFromThis<FLiveLinkHubSession>
@@ -124,6 +133,7 @@ public:
 		{
 			LiveLinkProvider->ResetTimecodeSettings(Client);
 			LiveLinkProvider->ResetCustomTimeStepSettings(Client);
+			LiveLinkProvider->DisconnectClient(Client);
 		}
 
 		{
@@ -134,6 +144,39 @@ public:
 		OnClientRemovedFromSessionDelegate.Broadcast(Client);
 	}
 
+	virtual void RemoveAllClients() override
+	{
+		check(IsInGameThread());
+
+		TSharedPtr<FLiveLinkHubProvider> LiveLinkProvider = FLiveLinkHub::Get()->GetLiveLinkProvider();
+
+		TSet<FLiveLinkHubClientId> Clients;
+		{
+			FReadScopeLock Locker(SessionDataLock);
+			Clients = CachedSessionClients;
+		}
+
+		if (LiveLinkProvider)
+		{
+			for (const FLiveLinkHubClientId& ClientId : Clients)
+			{
+				LiveLinkProvider->ResetTimecodeSettings(ClientId);
+				LiveLinkProvider->ResetCustomTimeStepSettings(ClientId);
+				LiveLinkProvider->DisconnectClient(ClientId);
+			}
+		}
+
+		{
+			FWriteScopeLock Locker(SessionDataLock);
+			CachedSessionClients.Reset();
+		}
+
+		for (const FLiveLinkHubClientId& ClientId : Clients)
+		{
+			OnClientRemovedFromSessionDelegate.Broadcast(ClientId);
+		}
+	}
+
 	virtual bool IsClientInSession(const FLiveLinkHubClientId& Client) override
 	{
 		FReadScopeLock Locker(SessionDataLock);
@@ -188,9 +231,30 @@ public:
 		SessionData->CustomTimeStepSettings = CustomTimeStepSettings;
 	}
 
+	virtual ELiveLinkTopologyMode GetTopologyMode() const override
+	{
+		FReadScopeLock Locker(SessionDataLock);
+		return SessionData->TopologyMode;
+	}
+
+	virtual void SetTopologyMode(ELiveLinkTopologyMode Mode) override
+	{
+		{
+			FWriteScopeLock Locker(SessionDataLock);
+			SessionData->TopologyMode = Mode;
+		}
+
+		RemoveAllClients();
+
+		if (const TSharedPtr<FLiveLinkHubProvider> LiveLinkProvider = FLiveLinkHub::Get()->GetLiveLinkProvider())
+		{
+			LiveLinkProvider->PostUpdateTopologyMode(Mode);
+		}
+	}
+
 	void AddRestoredClient(FLiveLinkHubUEClientInfo& InOutRestoredClientInfo)
 	{
-		if (const TSharedPtr<FLiveLinkHubProvider> LiveLinkProvider = FModuleManager::Get().GetModuleChecked<FLiveLinkHubModule>("LiveLinkHub").GetLiveLinkProvider())
+		if (const TSharedPtr<FLiveLinkHubProvider> LiveLinkProvider = FLiveLinkHub::Get()->GetLiveLinkProvider())
 		{
 			LiveLinkProvider->AddRestoredClient(InOutRestoredClientInfo);
 
diff --git a/Engine/Plugins/Animation/LiveLinkHub/Source/LiveLinkHub/Private/Session/LiveLinkHubSessionData.h b/Engine/Plugins/Animation/LiveLinkHub/Source/LiveLinkHub/Private/Session/LiveLinkHubSessionData.h
--- a/Engine/Plugins/Animation/LiveLinkHub/Source/LiveLinkHub/Private/Session/LiveLinkHubSessionData.h
+++ b/Engine/Plugins/Animation/LiveLinkHub/Source/LiveLinkHub/Private/Session/LiveLinkHubSessionData.h
@@ -65,6 +65,10 @@ public:
 	UPROPERTY()
 	bool bUseLiveLinkHubAsCustomTimeStepSource = false;
 
+	/** If the hub should be a hub or a spoke. Hubs can only connect to UE instances, while spokes can connect to other hubs. */
+	UPROPERTY()
+	ELiveLinkTopologyMode TopologyMode = ELiveLinkTopologyMode::Hub;
+
 	/** Generic / extensible storage for additional fields to be serialized with the session. */
 	UPROPERTY(Instanced) // Instanced so that the contained types are serialized inline.
 	TMap<TSubclassOf<ULiveLinkHubSessionExtraData>, TObjectPtr<ULiveLinkHubSessionExtraData>> ExtraDatas;
diff --git a/Engine/Plugins/Animation/LiveLinkHub/Source/LiveLinkHub/Private/Settings/LiveLinkHubSettings.h b/Engine/Plugins/Animation/LiveLinkHub/Source/LiveLinkHub/Private/Settings/LiveLinkHubSettings.h
--- a/Engine/Plugins/Animation/LiveLinkHub/Source/LiveLinkHub/Private/Settings/LiveLinkHubSettings.h
+++ b/Engine/Plugins/Animation/LiveLinkHub/Source/LiveLinkHub/Private/Settings/LiveLinkHubSettings.h
@@ -36,7 +36,7 @@ protected:
 	
 public:
 	/** If enabled, discovered clients will be automatically added to the current session. */
-	UPROPERTY(config, EditAnywhere, Category="LiveLinkHub", DisplayName = "Automatically add discovered clients")
+	UPROPERTY(config, EditAnywhere, Category="LiveLinkHub", DisplayName = "Auto Connect Mode")
 	ELiveLinkHubAutoConnectMode AutoConnectClients = ELiveLinkHubAutoConnectMode::All;
 
 	/** The size in megabytes to buffer when streaming a recording. */
diff --git a/Engine/Plugins/Animation/LiveLinkHub/Source/LiveLinkHub/Private/UI/Widgets/SLiveLinkHubTopologyModeSwitcher.h b/Engine/Plugins/Animation/LiveLinkHub/Source/LiveLinkHub/Private/UI/Widgets/SLiveLinkHubTopologyModeSwitcher.h
new file mode 100644
--- /dev/null
+++ b/Engine/Plugins/Animation/LiveLinkHub/Source/LiveLinkHub/Private/UI/Widgets/SLiveLinkHubTopologyModeSwitcher.h
@@ -0,0 +1,106 @@
+// Copyright Epic Games, Inc. All Rights Reserved.
+
+#pragma once
+
+
+#include "Widgets/SCompoundWidget.h"
+
+#include "LiveLinkHub.h"
+#include "Widgets/Images/SImage.h"
+#include "Widgets/Input/SButton.h"
+
+#define LOCTEXT_NAMESPACE "LiveLinkHubTopologyModeSwitcher"
+
+class SLiveLinkHubTopologyModeSwitcher : public SCompoundWidget
+{
+public:
+	SLATE_BEGIN_ARGS(SLiveLinkHubTopologyModeSwitcher)
+	{}
+	SLATE_END_ARGS()
+
+	void Construct(const FArguments& InArgs)
+	{
+		TSharedPtr<FLiveLinkHub> LiveLinkHub = FLiveLinkHub::Get();
+
+		const bool bReadOnly = !LiveLinkHub->CanSetTopologyMode();
+
+		const ELiveLinkTopologyMode Mode = LiveLinkHub->GetTopologyMode();
+
+		FText ModeTooltip = GetModeTooltip(Mode);
+
+		if (bReadOnly)
+		{
+			ModeTooltip = FText::Format(LOCTEXT("ReadOnlyToolTip", "{0}\nNote: Since the mode was set by a command line, it cannot be changed at runtime."), ModeTooltip);
+		}
+
+		ChildSlot
+		[
+			SAssignNew(ModeButton, SButton)
+			.ContentPadding(FMargin(0.0f, 2.0f))
+			.HAlign(HAlign_Center)
+			.VAlign(VAlign_Center)
+			.IsEnabled(!bReadOnly)
+			.ToolTipText(ModeTooltip)
+			.OnClicked(this, &SLiveLinkHubTopologyModeSwitcher::OnModeClicked)
+			[
+				SAssignNew(ButtonContent, STextBlock)
+				.Text(GetModeLabel(Mode))
+			]
+		];
+
+		LiveLinkHub->OnTopologyModeChanged().AddSP(this, &SLiveLinkHubTopologyModeSwitcher::OnModeChanged);
+	}
+
+	virtual ~SLiveLinkHubTopologyModeSwitcher()
+	{
+		if (TSharedPtr<FLiveLinkHub> LiveLinkHub = FLiveLinkHub::Get())
+		{
+			LiveLinkHub->OnTopologyModeChanged().RemoveAll(this);
+		}
+	}
+
+private:
+	/** Handles toggling the topology mode. */
+	FReply OnModeClicked()
+	{
+		FLiveLinkHub::Get()->ToggleTopologyMode();
+		return FReply::Handled();
+	}
+
+	/** Handles updating the button's text and tooltip when the topology mode changed. */
+	void OnModeChanged(ELiveLinkTopologyMode Mode)
+	{
+		UEnum* Enum = StaticEnum<ELiveLinkTopologyMode>();
+
+		ELiveLinkTopologyMode OtherMode = ELiveLinkTopologyMode::Spoke;
+		if (Mode == ELiveLinkTopologyMode::Spoke)
+		{
+			OtherMode = ELiveLinkTopologyMode::Hub;
+		}
+		
+		ButtonContent->SetText(GetModeLabel(Mode));
+		ModeButton->SetToolTipText(FText::Format(LOCTEXT("ButtonContentToolTip", "{0}\nClick to change to {1} mode"), GetModeTooltip(Mode), UEnum::GetDisplayValueAsText(OtherMode)));
+	}
+
+	/** Returns the label for the mode button. */
+	static FText GetModeLabel(ELiveLinkTopologyMode InMode)
+	{
+		return FText::Format(LOCTEXT("ModeLabel", "Mode: {0}"), UEnum::GetDisplayValueAsText(InMode));
+	}
+
+	/** Returns the tooltip for a mode. */
+	static FText GetModeTooltip(ELiveLinkTopologyMode InMode)
+	{
+		UEnum* Enum = StaticEnum<ELiveLinkTopologyMode>();
+		const int32 ModeEnumIndex = Enum->GetIndexByValue(static_cast<int64>(InMode));
+		return Enum->GetToolTipTextByIndex(ModeEnumIndex);
+	}
+
+private:
+	/** Button that's responsible for switching the topology mode of the app. */
+	TSharedPtr<SButton> ModeButton;
+	/** Text indicating the current mode. */
+	TSharedPtr<STextBlock> ButtonContent;
+};
+
+#undef LOCTEXT_NAMESPACE
\ No newline at end of file
diff --git a/Engine/Plugins/Animation/LiveLinkHub/Source/LiveLinkHub/Private/UI/Window/LiveLinkHubWindowController.cpp b/Engine/Plugins/Animation/LiveLinkHub/Source/LiveLinkHub/Private/UI/Window/LiveLinkHubWindowController.cpp
--- a/Engine/Plugins/Animation/LiveLinkHub/Source/LiveLinkHub/Private/UI/Window/LiveLinkHubWindowController.cpp
+++ b/Engine/Plugins/Animation/LiveLinkHub/Source/LiveLinkHub/Private/UI/Window/LiveLinkHubWindowController.cpp
@@ -12,6 +12,7 @@
 #include "HAL/PlatformApplicationMisc.h"
 #include "Misc/ConfigCacheIni.h"
 #include "UI/Widgets/LiveLinkHubMainTabController.h"
+#include "UI/Widgets/SLiveLinkHubTopologyModeSwitcher.h"
 #include "UI/Window/ModalWindowManager.h"
 #include "Widgets/Docking/SDockTab.h"
 
@@ -25,8 +26,6 @@ FLiveLinkHubWindowController::FLiveLinkHubWindowController()
 
 FLiveLinkHubWindowController::~FLiveLinkHubWindowController()
 {
-	FGlobalTabmanager::Get()->SaveAllVisualState();
-
 	if (RootWindow)
 	{
 		RootWindow->SetOnWindowClosed(nullptr);
@@ -111,7 +110,21 @@ void FLiveLinkHubWindowController::RestoreLayout(TSharedPtr<FAssetEditorToolkit>
 				.Text(FText::FromString(FLiveLinkHub::Get()->GetLiveLinkProvider()->GetProviderName()))
 		];
 
-	FGlobalTabmanager::Get()->FindExistingLiveTab(StandaloneTabId)->SetTitleBarRightContent(ProviderNameWidget);
+	TSharedRef<SWidget> RightMenuWidget = 
+		SNew(SHorizontalBox)
+		+ SHorizontalBox::Slot()
+		.AutoWidth()
+		[
+			ProviderNameWidget
+		]
+		+ SHorizontalBox::Slot()
+		.AutoWidth()
+		.Padding(4.0f)
+		[
+			SNew(SLiveLinkHubTopologyModeSwitcher)
+		];
+
+		FGlobalTabmanager::Get()->FindExistingLiveTab(StandaloneTabId)->SetTitleBarRightContent(RightMenuWidget);
 }
 
 TSharedPtr<FModalWindowManager> FLiveLinkHubWindowController::InitializeSlateApplication()
@@ -140,7 +153,8 @@ void FLiveLinkHubWindowController::OnWindowClosed(const TSharedRef<SWindow>& Win
 {
 	if (TSharedPtr<FLiveLinkHub> LiveLinkHub = FLiveLinkHub::Get())
 	{
-		LiveLinkHub->OnClose();
+		FGlobalTabmanager::Get()->SaveAllVisualState();
+		LiveLinkHub->CloseWindow(EAssetEditorCloseReason::CloseAllAssetEditors);
 	}
 
 	RootWindow.Reset();
