Fix CIS regression introduced by bad merge

diff --git a/Engine/Plugins/Animation/LiveLinkHub/Source/LiveLinkHub/Private/Clients/LiveLinkHubProvider.cpp b/Engine/Plugins/Animation/LiveLinkHub/Source/LiveLinkHub/Private/Clients/LiveLinkHubProvider.cpp
--- a/Engine/Plugins/Animation/LiveLinkHub/Source/LiveLinkHub/Private/Clients/LiveLinkHubProvider.cpp
+++ b/Engine/Plugins/Animation/LiveLinkHub/Source/LiveLinkHub/Private/Clients/LiveLinkHubProvider.cpp
@@ -306,7 +306,8 @@ void FLiveLinkHubProvider::HandleHubConnectMessage(const FLiveLinkHubConnectMess
 	if (!UE::LiveLink::Messaging::CanTransmitTo(Mode, Message.ClientInfo.TopologyMode))
 	{
 		UE_LOG(LogLiveLinkHub, Verbose, TEXT("Provider: Denying connection from %s since its mode is incompatible with this provider's."), *Context->GetSender().ToString());
-		SendMessage(FMessageEndpoint::MakeMessage<FLiveLinkHubDisconnectMessage>(GetProviderName(), GetMachineName()), Context->GetSender(), EMessageFlags::Reliable);
+		FLiveLinkHubDisconnectMessage DisconnectMessage{ GetProviderName(), GetMachineName() };
+		SendMessage(FMessageEndpoint::MakeMessage<FLiveLinkHubDisconnectMessage>(MoveTemp(DisconnectMessage)), Context->GetSender(), EMessageFlags::Reliable);
 		CloseConnection(Context->GetSender());
 		return;
 	}
