Fix uninitialized struct checker warning in LiveLinkHubMessages.h

#rnx

diff --git a/Engine/Plugins/Animation/LiveLinkHub/Source/LiveLinkHubMessaging/Public/LiveLinkHubMessages.h b/Engine/Plugins/Animation/LiveLinkHub/Source/LiveLinkHubMessaging/Public/LiveLinkHubMessages.h
--- a/Engine/Plugins/Animation/LiveLinkHub/Source/LiveLinkHubMessaging/Public/LiveLinkHubMessages.h
+++ b/Engine/Plugins/Animation/LiveLinkHub/Source/LiveLinkHubMessaging/Public/LiveLinkHubMessages.h
@@ -222,7 +222,7 @@ struct LIVELINKHUBMESSAGING_API FLiveLinkHubDiscoveryMessage
 	FString InstanceId;
 
 	/** Creation time used to calculate the machine time offset. */
-	UPROPERTY()
+	UPROPERTY(meta = (IgnoreForMemberInitializationTest))
 	double CreationTime = FPlatformTime::Seconds();
 
 	/** LiveLink Version in use by this client. */
