LiveLinkHub - Fix flickering source appearing in UE if there is a spoke instance running.
Fix MessageBus sources showing up in the LiveLinkHub section

#rnx

diff --git a/Engine/Plugins/Animation/LiveLinkHub/Source/LiveLinkHub/Private/Clients/LiveLinkHubProvider.cpp b/Engine/Plugins/Animation/LiveLinkHub/Source/LiveLinkHub/Private/Clients/LiveLinkHubProvider.cpp
--- a/Engine/Plugins/Animation/LiveLinkHub/Source/LiveLinkHub/Private/Clients/LiveLinkHubProvider.cpp
+++ b/Engine/Plugins/Animation/LiveLinkHub/Source/LiveLinkHub/Private/Clients/LiveLinkHubProvider.cpp
@@ -585,10 +585,13 @@ TMap<FName, FString> FLiveLinkHubProvider::GetAnnotations() const
 	// In case this is handled by the discovery manager (which doesn't directly handle LiveLinkHub messages)
 	AnnotationsCopy.Add(FLiveLinkHubMessageAnnotation::IdAnnotation, FLiveLinkHub::Get()->GetId().ToString());
 
+	// In case this is handled by the discovery manager (which doesn't directly handle LiveLinkHub messages)
+	const ELiveLinkTopologyMode Mode = FLiveLinkHub::Get()->GetTopologyMode();
+	FString TopopologyMode = StaticEnum<ELiveLinkTopologyMode>()->GetNameStringByValue(static_cast<int64>(Mode));
+	AnnotationsCopy.Add(FLiveLinkMessageAnnotation::TopologyModeAnnotation, MoveTemp(TopopologyMode));
+
 	TMap<FName, FString> BaseAnnotations = FLiveLinkProvider::GetAnnotations();
 	BaseAnnotations.Append(AnnotationsCopy);
-	// In case this is handled by the discovery manager (which doesn't directly handle LiveLinkHub messages)
-	AnnotationsCopy.Add(FLiveLinkHubMessageAnnotation::IdAnnotation, FLiveLinkHub::Get()->GetId().ToString());
 
 	return BaseAnnotations;
 }
diff --git a/Engine/Plugins/Animation/LiveLinkHub/Source/LiveLinkHubMessaging/Private/LiveLinkHubMessagingModule.cpp b/Engine/Plugins/Animation/LiveLinkHub/Source/LiveLinkHubMessaging/Private/LiveLinkHubMessagingModule.cpp
--- a/Engine/Plugins/Animation/LiveLinkHub/Source/LiveLinkHubMessaging/Private/LiveLinkHubMessagingModule.cpp
+++ b/Engine/Plugins/Animation/LiveLinkHub/Source/LiveLinkHubMessaging/Private/LiveLinkHubMessagingModule.cpp
@@ -89,7 +89,18 @@ ELiveLinkTopologyMode FLiveLinkHubMessagingModule::GetHostTopologyMode() const
 
 bool FLiveLinkHubMessagingModule::OnFilterMessageBusSource(UClass* FactoryClass, TSharedPtr<struct FProviderPollResult, ESPMode::ThreadSafe> PollResult)
 {
-	return LiveLinkHubConnectionManager::ShouldAcceptConnectionFrom(InstanceInfo.TopologyMode, PollResult, InstanceInfo.Id);
+	// Only display Hub/Spoke sources in "LiveLinkHub" section of the add source dropdown.
+	bool bValidTopologyMode = false;
+	if (FactoryClass == ULiveLinkHubMessageBusSourceFactory::StaticClass())
+	{
+		bValidTopologyMode = LiveLinkHubConnectionManager::GetPollResultTopologyMode(PollResult) == ELiveLinkTopologyMode::Hub || LiveLinkHubConnectionManager::GetPollResultTopologyMode(PollResult) == ELiveLinkTopologyMode::Spoke;
+	}
+	else if (FactoryClass == ULiveLinkMessageBusSourceFactory::StaticClass())
+	{
+		bValidTopologyMode = LiveLinkHubConnectionManager::GetPollResultTopologyMode(PollResult) != ELiveLinkTopologyMode::Hub;
+	}
+
+	return bValidTopologyMode && LiveLinkHubConnectionManager::ShouldAcceptConnectionFrom(InstanceInfo.TopologyMode, PollResult, InstanceInfo.Id);
 }
 
 void FLiveLinkHubMessagingModule::HandleDiscoveryMessage(const FLiveLinkHubDiscoveryMessage& Message, const TSharedRef<IMessageContext, ESPMode::ThreadSafe>& Context)
