LiveLinkHub: Enable selection and copy paste on automatic tokens under settings.

diff --git a/Engine/Plugins/Animation/LiveLinkHub/Source/LiveLinkHub/Private/Settings/LiveLinkHubSettingsCustomization.cpp b/Engine/Plugins/Animation/LiveLinkHub/Source/LiveLinkHub/Private/Settings/LiveLinkHubSettingsCustomization.cpp
--- a/Engine/Plugins/Animation/LiveLinkHub/Source/LiveLinkHub/Private/Settings/LiveLinkHubSettingsCustomization.cpp
+++ b/Engine/Plugins/Animation/LiveLinkHub/Source/LiveLinkHub/Private/Settings/LiveLinkHubSettingsCustomization.cpp
@@ -3,12 +3,12 @@
 #include "LiveLinkHubSettingsCustomization.h"
 
 #include "Config/LiveLinkHubTemplateTokens.h"
-#include "DetailCategoryBuilder.h"
 #include "DetailLayoutBuilder.h"
 #include "DetailWidgetRow.h"
 #include "Editor.h"
 #include "LiveLinkHubSettings.h"
 #include "NamingTokensEditorSubsystem.h"
+#include "Widgets/Input/SEditableText.h"
 #include "Widgets/Text/STextBlock.h"
 
 #define LOCTEXT_NAMESPACE "LiveLinkHubSettingsCustomization"
@@ -44,7 +44,8 @@ void FLiveLinkHubSettingsCustomization::CustomizeDetails(IDetailLayoutBuilder& D
 		+ SVerticalBox::Slot()
 		.AutoHeight()
 		[
-			SNew(STextBlock)
+			SNew(SEditableText)
+			.IsReadOnly(true)
 			.Text_Raw(this, &FLiveLinkHubSettingsCustomization::GetDisplayTokenText)
 			.Font(IDetailLayoutBuilder::GetDetailFont())
 		]
