LiveLinkHub - Add option to remove invalid subjects after loading session

Can be enabled through LiveLinkHub's setting: RemoveInvalidSubjectsAfterLoadingSession (default: false)

diff --git a/Engine/Plugins/Animation/LiveLinkHub/Source/LiveLinkHub/Private/Session/LiveLinkHubSessionManager.h b/Engine/Plugins/Animation/LiveLinkHub/Source/LiveLinkHub/Private/Session/LiveLinkHubSessionManager.h
--- a/Engine/Plugins/Animation/LiveLinkHub/Source/LiveLinkHub/Private/Session/LiveLinkHubSessionManager.h
+++ b/Engine/Plugins/Animation/LiveLinkHub/Source/LiveLinkHub/Private/Session/LiveLinkHubSessionManager.h
@@ -214,6 +214,7 @@ public:
 
 		IDesktopPlatform* DesktopPlatform = FDesktopPlatformModule::Get();
 		const void* ParentWindowWindowHandle = FSlateApplication::Get().FindBestParentWindowHandleForDialogs(nullptr);
+		OpenSessionTimestampsSeconds = FPlatformTime::Seconds();
 		const bool bFileSelected = DesktopPlatform->OpenFileDialog(
 			ParentWindowWindowHandle,
 			LOCTEXT("LiveLinkHubOpenTitle", "Open").ToString(),
@@ -429,6 +430,18 @@ private:
 			}
 
 			OnActiveSessionChangedDelegate.Broadcast(CurrentSessionPtr.ToSharedRef());
+
+			if (GetDefault<ULiveLinkHubSettings>()->bRemoveInvalidSubjectsAfterLoadingSession)
+			{
+				constexpr double ActualWaitTime = 1.5;
+
+				// This is odd, but since we're blocking the game thread for a while to choose a session file, the internal time of the TimerManager hasn't been updated yet.
+				// In order for our timer to have the correct start time, we need to offset our time correctly so that it matches the internal time, else our timer would execute immediately.
+				double AdjustedWaitTime = FPlatformTime::Seconds() - OpenSessionTimestampsSeconds + ActualWaitTime;
+
+				FTimerHandle Handle;
+				GEditor->GetTimerManager()->SetTimer(Handle, FTimerDelegate::CreateRaw(this, &FLiveLinkHubSessionManager::CullUnresponsiveSubjects), AdjustedWaitTime, false);
+			}
 		}
 	}
 
@@ -463,6 +476,27 @@ private:
 		OnActiveSessionChangedDelegate.Broadcast(CurrentSessionPtr.ToSharedRef());
 	}
 
+	/** Remove subjects that haven't received data in a while. */
+	void CullUnresponsiveSubjects()
+	{
+		FLiveLinkHubClient* LiveLinkHubClient = static_cast<FLiveLinkHubClient*>(&IModularFeatures::Get().GetModularFeature<ILiveLinkClient>(ILiveLinkClient::ModularFeatureName));
+		if (LiveLinkHubClient)
+		{
+			// We probably don't want to remove subjects that were explicitly marked as disabled.
+			constexpr bool bDisabledSubjects = false;
+			constexpr bool bVirtualSubjects = false;
+
+			TArray<FLiveLinkSubjectKey> Subjects = LiveLinkHubClient->GetSubjects(bDisabledSubjects, bVirtualSubjects);
+			for (const FLiveLinkSubjectKey& SubjectKey : Subjects)
+			{
+				if (LiveLinkHubClient->GetSubjectState(SubjectKey.SubjectName) == ELiveLinkSubjectState::Unresponsive)
+				{
+					LiveLinkHubClient->RemoveSubject_AnyThread(SubjectKey);
+				}
+			}
+		}
+	}
+
 private:
 	/** Session that holds the current configuration of the hub (Clients, sources, subjects). */
 	TSharedPtr<FLiveLinkHubSession> CurrentSession;
@@ -494,6 +528,9 @@ private:
 	/** Delegate triggered when recording ends. */
 	FSimpleMulticastDelegate OnRecordingStoppedDelegate;
 
+	/** Timestamp of when the open session button was clicked. */
+	double OpenSessionTimestampsSeconds = 0.0;
+
 	/** Critical section used to synchronize access to the current session. */
 	mutable FCriticalSection CurrentSessionCS;
 };
diff --git a/Engine/Plugins/Animation/LiveLinkHub/Source/LiveLinkHub/Private/Settings/LiveLinkHubSettings.h b/Engine/Plugins/Animation/LiveLinkHub/Source/LiveLinkHub/Private/Settings/LiveLinkHubSettings.h
--- a/Engine/Plugins/Animation/LiveLinkHub/Source/LiveLinkHub/Private/Settings/LiveLinkHubSettings.h
+++ b/Engine/Plugins/Animation/LiveLinkHub/Source/LiveLinkHub/Private/Settings/LiveLinkHubSettings.h
@@ -55,6 +55,10 @@ public:
 	UPROPERTY(config)
 	TArray<FName> ProjectSettingsToDisplay;
 
+	/** If this is enabled, invalid subjects will be removed after loading a session. */
+	UPROPERTY(config, EditAnywhere, AdvancedDisplay, Category = "LiveLinkHub")
+	bool bRemoveInvalidSubjectsAfterLoadingSession = false;
+
 	/**
 	 * - Experimental - If this is disabled, LiveLinkHub's LiveLink Client will tick outside of the game thread.
 	 * This allows processing LiveLink frame snapshots without the risk of being blocked by the game / ui thread.
