LiveLinkHub: Fix NamingTokens causing a crash when referenced in the editor and LLH isn't available.

diff --git a/Engine/Plugins/Animation/LiveLinkHub/Source/LiveLinkHub/Private/Config/LiveLinkHubTemplateTokens.cpp b/Engine/Plugins/Animation/LiveLinkHub/Source/LiveLinkHub/Private/Config/LiveLinkHubTemplateTokens.cpp
--- a/Engine/Plugins/Animation/LiveLinkHub/Source/LiveLinkHub/Private/Config/LiveLinkHubTemplateTokens.cpp
+++ b/Engine/Plugins/Animation/LiveLinkHub/Source/LiveLinkHub/Private/Config/LiveLinkHubTemplateTokens.cpp
@@ -27,9 +27,15 @@ void ULiveLinkHubNamingTokens::OnCreateDefaultTokens(TArray<FNamingTokenData>& T
 void ULiveLinkHubNamingTokens::OnPreEvaluate_Implementation(const FNamingTokensEvaluationData& InEvaluationData)
 {
 	Super::OnPreEvaluate_Implementation(InEvaluationData);
-	if (const TSharedPtr<ILiveLinkHubSessionManager> SessionManager = FModuleManager::Get().GetModuleChecked<FLiveLinkHubModule>("LiveLinkHub").GetLiveLinkHub()->GetSessionManager())
+	if (const FLiveLinkHubModule* LiveLinkHubModule = FModuleManager::Get().GetModulePtr<FLiveLinkHubModule>("LiveLinkHub"))
 	{
-		SessionName = FPaths::GetBaseFilename(SessionManager->GetLastConfigPath());
+		if (const TSharedPtr<FLiveLinkHub> LiveLinkHub = LiveLinkHubModule->GetLiveLinkHub())
+		{
+			if (const TSharedPtr<ILiveLinkHubSessionManager> SessionManager = LiveLinkHub->GetSessionManager())
+			{
+				SessionName = FPaths::GetBaseFilename(SessionManager->GetLastConfigPath());
+			}
+		}
 	}
 }
 
