Fix crash when shutting down livelinkhub while playing a recording

diff --git a/Engine/Plugins/Animation/LiveLinkHub/Source/LiveLinkHub/Private/Recording/Implementations/LiveLinkUAssetRecording.cpp b/Engine/Plugins/Animation/LiveLinkHub/Source/LiveLinkHub/Private/Recording/Implementations/LiveLinkUAssetRecording.cpp
--- a/Engine/Plugins/Animation/LiveLinkHub/Source/LiveLinkHub/Private/Recording/Implementations/LiveLinkUAssetRecording.cpp
+++ b/Engine/Plugins/Animation/LiveLinkHub/Source/LiveLinkHub/Private/Recording/Implementations/LiveLinkUAssetRecording.cpp
@@ -23,9 +23,24 @@
 
 ULiveLinkUAssetRecording::~ULiveLinkUAssetRecording()
 {
-	if (!IsTemplate() && !IsEngineExitRequested())
+	if (!IsTemplate())
 	{
-		UnloadRecordingData();
+		if (!IsEngineExitRequested())
+		{
+			UnloadRecordingData();
+		}
+		else
+		{
+			bCancelStream = true;
+			if (AsyncStreamTask.IsValid())
+			{
+				if (!AsyncStreamTask->Cancel())
+				{
+					AsyncStreamTask->EnsureCompletion();
+				}
+				AsyncStreamTask.Reset();
+			}
+		}
 	}
 }
 
diff --git a/Engine/Plugins/Animation/LiveLinkHub/Source/LiveLinkHub/Private/Recording/Implementations/LiveLinkUAssetRecordingPlayer.cpp b/Engine/Plugins/Animation/LiveLinkHub/Source/LiveLinkHub/Private/Recording/Implementations/LiveLinkUAssetRecordingPlayer.cpp
--- a/Engine/Plugins/Animation/LiveLinkHub/Source/LiveLinkHub/Private/Recording/Implementations/LiveLinkUAssetRecordingPlayer.cpp
+++ b/Engine/Plugins/Animation/LiveLinkHub/Source/LiveLinkHub/Private/Recording/Implementations/LiveLinkUAssetRecordingPlayer.cpp
@@ -357,38 +357,46 @@ void FLiveLinkUAssetRecordingPlayer::ShutdownPlayback()
 
 bool FLiveLinkUAssetRecordingPlayer::StreamPlayback(int32 InFromFrame)
 {
-	const int32 InitialFramesToBuffer = GetNumFramesToBuffer();
-	LoadedRecording->LoadRecordingData(InFromFrame, InitialFramesToBuffer);
-
-	// Make sure there are a few frames ready.
-	if (LoadedRecording->WaitForBufferedFrames(InFromFrame, InFromFrame + 2))
+	if (LoadedRecording.IsValid())
 	{
-		// On initial load, the correct frame size may not be calculated prior until after waiting for the buffer,
-		// update the correct number of frames and start buffering them.
-		const int32 CurrentFramesToBuffer = GetNumFramesToBuffer();
-		if (CurrentFramesToBuffer != InitialFramesToBuffer)
+		const int32 InitialFramesToBuffer = GetNumFramesToBuffer();
+		LoadedRecording->LoadRecordingData(InFromFrame, InitialFramesToBuffer);
+
+		// Make sure there are a few frames ready.
+		if (LoadedRecording->WaitForBufferedFrames(InFromFrame, InFromFrame + 2))
 		{
-			LoadedRecording->LoadRecordingData(InFromFrame, CurrentFramesToBuffer);
+			// On initial load, the correct frame size may not be calculated prior until after waiting for the buffer,
+			// update the correct number of frames and start buffering them.
+			const int32 CurrentFramesToBuffer = GetNumFramesToBuffer();
+			if (CurrentFramesToBuffer != InitialFramesToBuffer)
+			{
+				LoadedRecording->LoadRecordingData(InFromFrame, CurrentFramesToBuffer);
+			}
+
+			// Take the available recording data.
+			LoadedRecording->CopyRecordingData(CurrentRecordingPlayback);
+			return true;
 		}
-
-		// Take the available recording data.
-		LoadedRecording->CopyRecordingData(CurrentRecordingPlayback);
-		return true;
 	}
 	return false;
 }
 
 int32 FLiveLinkUAssetRecordingPlayer::GetNumFramesToBuffer() const
 {
-	const int32 FrameSize = LoadedRecording->GetFrameDiskSize();
-	const int32 MaxFrameBufferSizeMB = GetDefault<ULiveLinkHubSettings>()->FrameBufferSizeMB;
-	const int32 MaxFrameBufferSizeBytes = MaxFrameBufferSizeMB * 1024 * 1024;
+	int32 TotalFramesToBuffer = 0;
+	if (LoadedRecording.IsValid())
+	{
+		const int32 FrameSize = LoadedRecording->GetFrameDiskSize();
+		const int32 MaxFrameBufferSizeMB = GetDefault<ULiveLinkHubSettings>()->FrameBufferSizeMB;
+		const int32 MaxFrameBufferSizeBytes = MaxFrameBufferSizeMB * 1024 * 1024;
 
-	// We divide total frames by 2, since they get doubled later to account for scrubbing in both directions.
-	int32 TotalFramesToBuffer = FrameSize > 0 ? MaxFrameBufferSizeBytes / FrameSize / 2 : 0;
+		// We divide total frames by 2, since they get doubled later to account for scrubbing in both directions.
+		TotalFramesToBuffer = FrameSize > 0 ? MaxFrameBufferSizeBytes / FrameSize / 2 : 0;
+
+		// Ensure at least a few frames can be buffered.
+		constexpr int32 MinFrames = 3;
+		TotalFramesToBuffer = FMath::Max(TotalFramesToBuffer, MinFrames);
+	}
 
-	// Ensure at least a few frames can be buffered.
-	constexpr int32 MinFrames = 3;
-	TotalFramesToBuffer = FMath::Max(TotalFramesToBuffer, MinFrames);
 	return TotalFramesToBuffer;
 }
