LiveLinkhub - Default autoconnect mode to Local rather than global
Add autoconnect icons
Fix LLH settings not getting saved when changing the mode through the UI

diff --git a/Engine/Plugins/Animation/LiveLink/Source/LiveLinkEditor/Private/LiveLinkEditorModule.cpp b/Engine/Plugins/Animation/LiveLink/Source/LiveLinkEditor/Private/LiveLinkEditorModule.cpp
--- a/Engine/Plugins/Animation/LiveLink/Source/LiveLinkEditor/Private/LiveLinkEditorModule.cpp
+++ b/Engine/Plugins/Animation/LiveLink/Source/LiveLinkEditor/Private/LiveLinkEditorModule.cpp
@@ -110,7 +110,9 @@ public:
 		StyleSet->Set("LiveLinkHub.Playback.Icon", new IMAGE_PLUGIN_BRUSH_SVG(TEXT("Starship/Playback"), Icon16x16));
 		StyleSet->Set("LiveLinkHub.Subjects.Icon", new IMAGE_PLUGIN_BRUSH_SVG(TEXT("Starship/Subjects"), Icon16x16));
 
-		StyleSet->Set("LiveLinkHub.AutoConnect", new IMAGE_PLUGIN_BRUSH_SVG("Starship/AutoConnect", Icon16x16));
+		StyleSet->Set("LiveLinkHub.AutoConnect.Disabled", new IMAGE_PLUGIN_BRUSH_SVG("Starship/AutoConnectDisabled", Icon16x16));
+		StyleSet->Set("LiveLinkHub.AutoConnect.Global", new IMAGE_PLUGIN_BRUSH_SVG("Starship/AutoConnectGlobal", Icon16x16));
+		StyleSet->Set("LiveLinkHub.AutoConnect.Local", new IMAGE_PLUGIN_BRUSH_SVG("Starship/AutoConnectLocal", Icon16x16));
 
 		StyleSet->Set("LiveLinkHub.Icon.Small", new IMAGE_PLUGIN_BRUSH(TEXT("Starship/LiveLinkHub_16x16"), Icon16x16));
 		StyleSet->Set("LiveLinkHub.Icon", new IMAGE_PLUGIN_BRUSH_SVG(TEXT("Starship/LiveLinkHub_Icon"), Icon40x40));
diff --git a/Engine/Plugins/Animation/LiveLinkHub/Source/LiveLinkHub/Private/Clients/SLiveLinkHubClientsView.h b/Engine/Plugins/Animation/LiveLinkHub/Source/LiveLinkHub/Private/Clients/SLiveLinkHubClientsView.h
--- a/Engine/Plugins/Animation/LiveLinkHub/Source/LiveLinkHub/Private/Clients/SLiveLinkHubClientsView.h
+++ b/Engine/Plugins/Animation/LiveLinkHub/Source/LiveLinkHub/Private/Clients/SLiveLinkHubClientsView.h
@@ -43,34 +43,34 @@ DECLARE_DELEGATE_OneParam(FOnRemoveClientFromSession, FLiveLinkHubClientId/*Clie
 static const FName NameColumnId = "Name";
 static const FName StatusColumnId = "Status";
 static const FName EnabledIconColumnId = "EnabledIcon";
+static const FName DisabledIconName = "LiveLinkHub.AutoConnect.Disabled";
+static const FName GlobalIconName = "LiveLinkHub.AutoConnect.Global";
+static const FName LocalOnlyIconName = "LiveLinkHub.AutoConnect.Local";
 
 namespace LiveLinkHubClientsViewUtils
 {
-	static TSharedRef<SWidget> GetAutoConnectIcon(ELiveLinkHubAutoConnectMode Mode)
+	static FSlateIcon GetAutoConnectIcon(ELiveLinkHubAutoConnectMode Mode)
 	{
-		FSlateColor IconColor;
+		FName IconName;
 
-		// Note: Temporary until we have icons for autoconnect modes.
 		switch (Mode)
 		{
 			case ELiveLinkHubAutoConnectMode::Disabled:
-				IconColor = FSlateColor::UseSubduedForeground();
+				IconName = DisabledIconName;
 				break;
 			
 			case ELiveLinkHubAutoConnectMode::All:
-				IconColor = FLinearColor::Yellow;
+				IconName = GlobalIconName;
 				break;
 
 			case ELiveLinkHubAutoConnectMode::LocalOnly:
-				IconColor = FSlateColor::UseForeground();
+				IconName = LocalOnlyIconName;
 				break;
 			default:
 				checkNoEntry();
 		}
 
-		return SNew(SImage)
-		.ColorAndOpacity(IconColor)
-		.Image(FSlateIcon("LiveLinkStyle", "LiveLinkHub.AutoConnect").GetIcon());
+		return FSlateIcon("LiveLinkStyle", IconName);
 	}
 }
 
@@ -876,7 +876,9 @@ private:
 			for (int32 Index = 0; Index < Enum->NumEnums() - 1; Index++)
 			{
 				const ELiveLinkHubAutoConnectMode Mode = (ELiveLinkHubAutoConnectMode)Enum->GetValueByIndex(Index);
-				MenuBuilder.AddMenuEntry(Enum->GetDisplayNameTextByIndex(Index), Enum->GetToolTipTextByIndex(Index), FSlateIcon("LiveLinkStyle", "LiveLinkHub.AutoConnect"), FExecuteAction::CreateRaw(this, &SLiveLinkHubClientsView::SetAutoConnectMode, Mode));
+
+				const FSlateIcon Icon = LiveLinkHubClientsViewUtils::GetAutoConnectIcon(Mode);
+				MenuBuilder.AddMenuEntry(Enum->GetDisplayNameTextByIndex(Index), Enum->GetToolTipTextByIndex(Index), Icon, FExecuteAction::CreateRaw(this, &SLiveLinkHubClientsView::SetAutoConnectMode, Mode));
 			}
 
 			MenuBuilder.EndSection();
@@ -889,7 +891,7 @@ private:
 	void SetAutoConnectMode(ELiveLinkHubAutoConnectMode AutoConnectMode)
 	{
 		GetMutableDefault<ULiveLinkHubSettings>()->AutoConnectClients = AutoConnectMode;
-		GetMutableDefault<ULiveLinkHubSettings>()->SaveConfig();
+		GetMutableDefault<ULiveLinkHubSettings>()->TryUpdateDefaultConfigFile();
 	}
 
 	/** Create the widget for switching autoconnect mode. */
@@ -908,9 +910,12 @@ private:
 		{
 			const int32 Value = Enum->GetValueByIndex(Index);
 
+			const FSlateIcon Icon = LiveLinkHubClientsViewUtils::GetAutoConnectIcon((ELiveLinkHubAutoConnectMode)Value);
+
 			Switcher->AddSlot()
 				[
-					LiveLinkHubClientsViewUtils::GetAutoConnectIcon((ELiveLinkHubAutoConnectMode)Value)
+					SNew(SImage)
+						.Image(Icon.GetIcon())
 				];
 		}
 		return Switcher;
diff --git a/Engine/Plugins/Animation/LiveLinkHub/Source/LiveLinkHub/Private/Settings/LiveLinkHubSettings.h b/Engine/Plugins/Animation/LiveLinkHub/Source/LiveLinkHub/Private/Settings/LiveLinkHubSettings.h
--- a/Engine/Plugins/Animation/LiveLinkHub/Source/LiveLinkHub/Private/Settings/LiveLinkHubSettings.h
+++ b/Engine/Plugins/Animation/LiveLinkHub/Source/LiveLinkHub/Private/Settings/LiveLinkHubSettings.h
@@ -37,7 +37,7 @@ protected:
 public:
 	/** If enabled, discovered clients will be automatically added to the current session. */
 	UPROPERTY(config, EditAnywhere, Category="LiveLinkHub", DisplayName = "Auto Connect Mode")
-	ELiveLinkHubAutoConnectMode AutoConnectClients = ELiveLinkHubAutoConnectMode::All;
+	ELiveLinkHubAutoConnectMode AutoConnectClients = ELiveLinkHubAutoConnectMode::LocalOnly;
 
 	/** The size in megabytes to buffer when streaming a recording. */
 	UPROPERTY(config, EditAnywhere, Category="LiveLinkHub", meta = (ClampMin = "1", UIMin = "1"))
diff --git a/Engine/Source/Programs/LiveLinkHub/Config/DefaultEngine.ini b/Engine/Source/Programs/LiveLinkHub/Config/DefaultEngine.ini
--- a/Engine/Source/Programs/LiveLinkHub/Config/DefaultEngine.ini
+++ b/Engine/Source/Programs/LiveLinkHub/Config/DefaultEngine.ini
@@ -21,7 +21,7 @@ bRebroadcastSubjectsByDefault=true
 bAllowEditingRebroadcastProperty=false
 
 [/Script/LiveLinkHub.LiveLinkHubSettings]
-bAutoAddDiscoveredClients=true
+AutoConnectClients=LocalOnly
 +ProjectSettingsToDisplay=Live Link
 +ProjectSettingsToDisplay=Live Link Hub
 +ProjectSettingsToDisplay=UDPMessaging
