Fix other crash when exiting LiveLinkHub while playing a recording

diff --git a/Engine/Plugins/Animation/LiveLinkHub/Source/LiveLinkHub/Private/Recording/LiveLinkHubPlaybackController.cpp b/Engine/Plugins/Animation/LiveLinkHub/Source/LiveLinkHub/Private/Recording/LiveLinkHubPlaybackController.cpp
--- a/Engine/Plugins/Animation/LiveLinkHub/Source/LiveLinkHub/Private/Recording/LiveLinkHubPlaybackController.cpp
+++ b/Engine/Plugins/Animation/LiveLinkHub/Source/LiveLinkHub/Private/Recording/LiveLinkHubPlaybackController.cpp
@@ -75,7 +75,12 @@ FLiveLinkHubPlaybackController::FLiveLinkHubPlaybackController()
 FLiveLinkHubPlaybackController::~FLiveLinkHubPlaybackController()
 {
 	bIsDestructing = true;
-	Eject();
+	bIsPlaying = false;
+	if (!IsEngineExitRequested())
+	{
+		// Avoid doing a bunch of unnecessary work if the engine is currently exiting.
+		Eject();
+	}
 	Stopping = true;
 	PlaybackEvent->Trigger();
 
@@ -506,6 +511,13 @@ uint32 FLiveLinkHubPlaybackController::Run()
 		bIsPlaybackWaiting = false;
 		while (bIsPlaying)
 		{
+			if (IsEngineExitRequested())
+			{
+				Stopping = true;
+				bIsDestructing = true;
+				break;
+			}
+
 			if (bIsPaused)
 			{
 				FPlatformProcess::Sleep(0.002);
@@ -550,8 +562,11 @@ uint32 FLiveLinkHubPlaybackController::Run()
 		// If the loop ended because the recording is over.
 		bIsPlaying = false;
 		
-		// Trigger the playback finished delegate on the game thread.
-		FSimpleDelegateGraphTask::CreateAndDispatchWhenReady(FSimpleDelegateGraphTask::FDelegate::CreateRaw(this, &FLiveLinkHubPlaybackController::OnPlaybackFinished_Internal), TStatId(), nullptr, ENamedThreads::GameThread);
+		if (!bIsDestructing)
+		{
+			// Trigger the playback finished delegate on the game thread.
+			FSimpleDelegateGraphTask::CreateAndDispatchWhenReady(FSimpleDelegateGraphTask::FDelegate::CreateRaw(this, &FLiveLinkHubPlaybackController::OnPlaybackFinished_Internal), TStatId(), nullptr, ENamedThreads::GameThread);
+		}
 	}
 
 	return 0;
