Fix crash when deleting a disconnected client from LiveLinkHub

diff --git a/Engine/Plugins/Animation/LiveLinkHub/Source/LiveLinkHub/Private/Clients/LiveLinkHubProvider.cpp b/Engine/Plugins/Animation/LiveLinkHub/Source/LiveLinkHub/Private/Clients/LiveLinkHubProvider.cpp
--- a/Engine/Plugins/Animation/LiveLinkHub/Source/LiveLinkHub/Private/Clients/LiveLinkHubProvider.cpp
+++ b/Engine/Plugins/Animation/LiveLinkHub/Source/LiveLinkHub/Private/Clients/LiveLinkHubProvider.cpp
@@ -171,7 +171,7 @@ void FLiveLinkHubProvider::DisconnectClient(const FLiveLinkHubClientId& Client)
 	FMessageAddress* TargetAddress = AllAddresses.FindByPredicate(
 		[this, Client](const FMessageAddress& Address)
 		{
-			return AddressToIdCache.FindChecked(Address) == Client;
+			return AddressToClientId(Address) == Client;
 		});
 
 	// Todo: We may want to remove the client info and tracked address.
@@ -238,7 +238,7 @@ void FLiveLinkHubProvider::SendCustomTimeStepSettings(const FLiveLinkHubCustomTi
 
 		FMessageAddress* TargetAddress = AllAddresses.FindByPredicate([this, ClientId](const FMessageAddress& Address)
 			{
-				return AddressToIdCache.FindChecked(Address) == ClientId;
+				return AddressToClientId(Address) == ClientId;
 			});
 
 		if (TargetAddress)
@@ -559,6 +559,21 @@ void FLiveLinkHubProvider::CloseConnections(const TArray<FMessageAddress>& Close
 	}
 }
 
+FLiveLinkHubClientId FLiveLinkHubProvider::AddressToClientId(const FMessageAddress& Address) const
+{
+	FLiveLinkHubClientId ClientId;
+	if (const FLiveLinkHubClientId* FoundId = AddressToIdCache.Find(Address))
+	{
+		ClientId = *FoundId;
+	}
+	else
+	{
+		UE_LOG(LogLiveLinkHub, Warning, TEXT("Could not find a client for address %s."), *Address.ToString());
+	}
+	
+	return ClientId;
+}
+
 TArray<FLiveLinkHubClientId> FLiveLinkHubProvider::GetSessionClients() const
 {
 	TArray<FLiveLinkHubClientId> SessionClients;
diff --git a/Engine/Plugins/Animation/LiveLinkHub/Source/LiveLinkHub/Private/Clients/LiveLinkHubProvider.h b/Engine/Plugins/Animation/LiveLinkHub/Source/LiveLinkHub/Private/Clients/LiveLinkHubProvider.h
--- a/Engine/Plugins/Animation/LiveLinkHub/Source/LiveLinkHub/Private/Clients/LiveLinkHubProvider.h
+++ b/Engine/Plugins/Animation/LiveLinkHub/Source/LiveLinkHub/Private/Clients/LiveLinkHubProvider.h
@@ -134,6 +134,11 @@ private:
 
 	void CloseConnections(const TArray<FMessageAddress>& ClosedAddresses);
 
+	/** 
+	 * Get the client id that corresponds to this address from our cache. 
+	 * May return an invalid ID if the address is not in the cache (ie. If client is disconnecting))
+	 **/
+	FLiveLinkHubClientId AddressToClientId(const FMessageAddress& Address) const;
 protected:
 	//~ Begin ILiveLinkHubClientsModel interface
 	virtual void OnConnectionsClosed(const TArray<FMessageAddress>& ClosedAddresses) override;
