Fix recordings not using previous outbound names

This was caused by the subject settings getting re-initialized after the recording preset is loaded. We now prevent the settings from being initialized if its subject key is already set.

Also allow python to initialize asynchronously by default to speed up launch time.

diff --git a/Engine/Plugins/Animation/LiveLinkHub/Source/LiveLinkHub/Private/LiveLinkHubSubjectSettings.cpp b/Engine/Plugins/Animation/LiveLinkHub/Source/LiveLinkHub/Private/LiveLinkHubSubjectSettings.cpp
--- a/Engine/Plugins/Animation/LiveLinkHub/Source/LiveLinkHub/Private/LiveLinkHubSubjectSettings.cpp
+++ b/Engine/Plugins/Animation/LiveLinkHub/Source/LiveLinkHub/Private/LiveLinkHubSubjectSettings.cpp
@@ -10,13 +10,16 @@
 
 void ULiveLinkHubSubjectSettings::Initialize(FLiveLinkSubjectKey InSubjectKey)
 {
-	ILiveLinkClient* LiveLinkClient = &IModularFeatures::Get().GetModularFeature<ILiveLinkClient>(ILiveLinkClient::ModularFeatureName);
-	SubjectName = InSubjectKey.SubjectName.ToString();
-	Key = InSubjectKey;
+	if (Key != InSubjectKey)
+	{
+		ILiveLinkClient* LiveLinkClient = &IModularFeatures::Get().GetModularFeature<ILiveLinkClient>(ILiveLinkClient::ModularFeatureName);
+		SubjectName = InSubjectKey.SubjectName.ToString();
+		Key = InSubjectKey;
 
-	OutboundName = SubjectName;
+		OutboundName = SubjectName;
 
-	Source = LiveLinkClient->GetSourceType(InSubjectKey.Source).ToString();
+		Source = LiveLinkClient->GetSourceType(InSubjectKey.Source).ToString();
+	}
 }
 
 void ULiveLinkHubSubjectSettings::PreEditChange(FProperty* PropertyAboutToChange)
diff --git a/Engine/Source/Programs/LiveLinkHub/Source/LiveLinkHubLauncher/Private/LiveLinkHubRun.cpp b/Engine/Source/Programs/LiveLinkHub/Source/LiveLinkHubLauncher/Private/LiveLinkHubRun.cpp
--- a/Engine/Source/Programs/LiveLinkHub/Source/LiveLinkHubLauncher/Private/LiveLinkHubRun.cpp
+++ b/Engine/Source/Programs/LiveLinkHub/Source/LiveLinkHubLauncher/Private/LiveLinkHubRun.cpp
@@ -92,7 +92,7 @@ int32 RunLiveLinkHub(const TCHAR* CommandLine)
 	});
 
 	// Start up the main loop, adding some extra command line arguments:
-	const TCHAR* const ExtraArgs = TEXT("-xrtrackingonly");
+	const TCHAR* const ExtraArgs = TEXT("-xrtrackingonly -dpcvars=Engine.Python.ForceWrapperGenerationOnSingleThread=false");
 	int32 Result = GEngineLoop.PreInit(*FString::Printf(TEXT("%s %s %s"),
 		*FPaths::GetProjectFilePath(), CommandLine, ExtraArgs));
 
