LiveLinkHub - Add support for mixing tab spawners between different modes
Add LiveLinkHubApplicationBase to expose a limited interface to application modes

diff --git a/Engine/Plugins/Animation/LiveLinkHub/Source/LiveLinkHub/Private/LiveLinkHub.cpp b/Engine/Plugins/Animation/LiveLinkHub/Source/LiveLinkHub/Private/LiveLinkHub.cpp
--- a/Engine/Plugins/Animation/LiveLinkHub/Source/LiveLinkHub/Private/LiveLinkHub.cpp
+++ b/Engine/Plugins/Animation/LiveLinkHub/Source/LiveLinkHub/Private/LiveLinkHub.cpp
@@ -160,19 +160,22 @@ void FLiveLinkHub::InitializePostEngineInit()
 	// Update the mode in our connection manager as well.
 	MessagingModule.SetInstanceId(InstanceId);
 
-	AddLiveLinkHubApplicationMode("CreatorMode", MakeShared<FLiveLinkHubCreatorAppMode>(SharedThis<FWorkflowCentricApplication>(this)));
+	IModularFeatures::Get().OnModularFeatureRegistered().AddSP(this, &FLiveLinkHub::OnModularFeatureRegistered);
+
+	WindowController = MakeShared<FLiveLinkHubWindowController>();
+	
+	WindowController->RestoreLayout(AsShared());
+	
+	// Note: Registering modes *must* happen after the layout was restored, since RestoreLayout will initiate the asset editor. 
+	// Otherwise, the mode tabs will get registered with the global tab manager instead of the host app's tab manager.
+	AddLiveLinkHubApplicationMode("CreatorMode", MakeShared<FLiveLinkHubCreatorAppMode>(SharedThis<FLiveLinkHubApplicationBase>(this)));
 
 	for (ILiveLinkHubApplicationModeFactory* Factory : IModularFeatures::Get().GetModularFeatureImplementations<ILiveLinkHubApplicationModeFactory>(ILiveLinkHubApplicationModeFactory::ModularFeatureName))
 	{
-		TSharedRef<FLiveLinkHubApplicationMode> AppMode = Factory->CreateApplicationMode(SharedThis<FWorkflowCentricApplication>(this));
+		TSharedRef<FLiveLinkHubApplicationMode> AppMode = Factory->CreateLiveLinkHubAppMode(SharedThis<FLiveLinkHubApplicationBase>(this));
 		AddLiveLinkHubApplicationMode(AppMode->GetModeName(), MoveTemp(AppMode));
 	}
 
-	IModularFeatures::Get().OnModularFeatureRegistered().AddSP(this, &FLiveLinkHub::OnModularFeatureRegistered);
-
-	WindowController = MakeShared<FLiveLinkHubWindowController>();
-	
-	WindowController->RestoreLayout(AsShared());
 	SetCurrentMode("CreatorMode");
 
 	LiveLinkHubClient->OnStaticDataReceived_AnyThread().AddSP(this, &FLiveLinkHub::OnStaticDataReceived_AnyThread);
@@ -413,7 +416,7 @@ void FLiveLinkHub::OnModularFeatureRegistered(const FName& Type, IModularFeature
 	{
 		ILiveLinkHubApplicationModeFactory* Factory = static_cast<ILiveLinkHubApplicationModeFactory*>(ModularFeature);
 
-		TSharedRef<FLiveLinkHubApplicationMode> AppMode = Factory->CreateApplicationMode(SharedThis<FWorkflowCentricApplication>(this));
+		TSharedRef<FLiveLinkHubApplicationMode> AppMode = Factory->CreateLiveLinkHubAppMode(SharedThis<FLiveLinkHubApplicationBase>(this));
 		AddLiveLinkHubApplicationMode(AppMode->GetModeName(), MoveTemp(AppMode));
 	}
 }
diff --git a/Engine/Plugins/Animation/LiveLinkHub/Source/LiveLinkHub/Private/LiveLinkHubApplication.cpp b/Engine/Plugins/Animation/LiveLinkHub/Source/LiveLinkHub/Private/LiveLinkHubApplication.cpp
--- a/Engine/Plugins/Animation/LiveLinkHub/Source/LiveLinkHub/Private/LiveLinkHubApplication.cpp
+++ b/Engine/Plugins/Animation/LiveLinkHub/Source/LiveLinkHub/Private/LiveLinkHubApplication.cpp
@@ -21,7 +21,7 @@
 #define LOCTEXT_NAMESPACE "LiveLinkHubApplication"
 
 FLiveLinkHubApplication::FLiveLinkHubApplication()
-	: FWorkflowCentricApplication()
+	: FLiveLinkHubApplicationBase()
 {
 	LayoutIni = TEXT("LiveLinkHubLayout");
 
@@ -53,12 +53,72 @@ void FLiveLinkHubApplication::AddToolbarExtenders(FToolBarBuilder&)
 void FLiveLinkHubApplication::AddLiveLinkHubApplicationMode(FName ModeName, TSharedRef<FLiveLinkHubApplicationMode> Mode)
 {
 	CachedModeInfo.Add(ModeName, FLiveLinkHubAppModeInfo{ Mode->GetModeIcon(), Mode->GetDisplayName()});
+
+	TabManager->AddLocalWorkspaceMenuItem(Mode->GetWorkspaceMenuCategory());
+
+	AddApplicationMode(ModeName, MoveTemp(Mode));
+}
+
+void FLiveLinkHubApplication::PushTabFactories(const FWorkflowAllowedTabSet& FactorySetToPush, TSharedPtr<FLiveLinkHubApplicationMode> ApplicationMode)
+{
+	check(TabManager.IsValid());
+
+	for (auto FactoryIt = FactorySetToPush.CreateConstIterator(); FactoryIt; ++FactoryIt)
+	{
+		FactoryIt.Value()->RegisterTabSpawner(TabManager.ToSharedRef(), ApplicationMode.Get());
+	}
+}
+
+void FLiveLinkHubApplication::AddApplicationMode(FName ModeName, TSharedRef<FApplicationMode> Mode)
+{
+	// Register tabs for a mode once it's registered so that other modes have access to that tab.
+	Mode->RegisterTabFactories(TabManager);
 	FWorkflowCentricApplication::AddApplicationMode(ModeName, MoveTemp(Mode));
 }
 
 void FLiveLinkHubApplication::SetCurrentMode(FName NewMode)
 {
-	FWorkflowCentricApplication::SetCurrentMode(NewMode);
+	// Behaves like FWorkflowCentricApplication, but does not clear or add tab spawners, allowing us to share tabs between modes.
+	const bool bModeAlreadyActive = CurrentAppModePtr.IsValid() && (NewMode == CurrentAppModePtr->GetModeName());
+
+	if (!bModeAlreadyActive)
+	{
+		check(TabManager.IsValid());
+
+		const TSharedPtr<FApplicationMode> NewModePtr = GetApplicationModeList().FindRef(NewMode);
+
+		LayoutExtenders.Reset();
+
+		if (NewModePtr.IsValid())
+		{
+			if (NewModePtr->LayoutExtender.IsValid())
+			{
+				LayoutExtenders.Add(NewModePtr->LayoutExtender);
+			}
+			
+			// Deactivate the old mode
+			if (CurrentAppModePtr.IsValid())
+			{
+				CurrentAppModePtr->PreDeactivateMode();
+				CurrentAppModePtr->DeactivateMode(TabManager);
+				RemoveToolbarExtender(CurrentAppModePtr->GetToolbarExtender());
+				RemoveAllToolbarWidgets();
+			}
+
+			CurrentAppModePtr = NewModePtr;
+
+			// Activate the new layout
+			const TSharedRef<FTabManager::FLayout> NewLayout = CurrentAppModePtr->ActivateMode(TabManager);
+			RestoreFromLayout(NewLayout);
+
+			// Give the new mode a chance to do init
+			CurrentAppModePtr->PostActivateMode();
+
+			AddToolbarExtender(NewModePtr->GetToolbarExtender());
+			RegenerateMenusAndToolbars();
+		}
+	}
+
 	AppModeChangedDelegate.Broadcast(NewMode);
 }
 
diff --git a/Engine/Plugins/Animation/LiveLinkHub/Source/LiveLinkHub/Private/LiveLinkHubApplication.h b/Engine/Plugins/Animation/LiveLinkHub/Source/LiveLinkHub/Private/LiveLinkHubApplication.h
--- a/Engine/Plugins/Animation/LiveLinkHub/Source/LiveLinkHub/Private/LiveLinkHubApplication.h
+++ b/Engine/Plugins/Animation/LiveLinkHub/Source/LiveLinkHub/Private/LiveLinkHubApplication.h
@@ -2,6 +2,7 @@
 
 #pragma once
 
+#include "LiveLinkHubApplicationBase.h"
 #include "WorkflowOrientedApp/WorkflowCentricApplication.h"
 #include "Math/Color.h"
 #include "Misc/Optional.h"
@@ -17,11 +18,10 @@ struct FLiveLinkHubAppModeInfo
 
 DECLARE_MULTICAST_DELEGATE_OneParam(FOnApplicationModeChanged, /*CurrentMode*/FName);
 
-class FLiveLinkHubApplication : public FWorkflowCentricApplication
+class FLiveLinkHubApplication : public FLiveLinkHubApplicationBase
 {
 public:
 	FLiveLinkHubApplication();
-	virtual ~FLiveLinkHubApplication() {}
 
 	/** Get the delegate called when the app mode has changed. */
 	FOnApplicationModeChanged& OnApplicationModeChanged() { return AppModeChangedDelegate; }
@@ -33,7 +33,10 @@ public:
 	TArray<FName> GetApplicationModes() const;
 
 	/** Add a LiveLinkHub application mode. */
-	void AddLiveLinkHubApplicationMode(FName ModeName, TSharedRef<FLiveLinkHubApplicationMode> Mode);
+	virtual void AddLiveLinkHubApplicationMode(FName ModeName, TSharedRef<FLiveLinkHubApplicationMode> Mode);
+
+	/** Push tab factories for the given mode. */
+	virtual void PushTabFactories(const FWorkflowAllowedTabSet& FactorySetToPush, TSharedPtr<FLiveLinkHubApplicationMode> ApplicationMode);
 
 	//~ Begin FWorkflowCentricApplication interface
 	virtual void SetCurrentMode(FName NewMode) override;
@@ -59,6 +62,9 @@ protected:
 	//~ End AssetEditorToolkit interface
 
 private:
+	/** Don't allow LiveLinkHub to call this directly since it wouldn't register the cached mode info. */
+	virtual void AddApplicationMode(FName ModeName, TSharedRef<FApplicationMode> Mode) override;
+
 	/** Register and populate the file menu with the config commands. */
 	void RegisterFileMenu(UToolMenu* Menu);
 	/** Create the main menu bar. */
@@ -71,8 +77,6 @@ private:
 	void AddToolbarExtenders(FToolBarBuilder&);
 
 private:
-	/** Don't allow LiveLinkHub to call this directly since it wouldn't register the cached mode info. */
-	using FWorkflowCentricApplication::AddApplicationMode;
 
 	/** Delegate called when the active app mode has changed. */
 	FOnApplicationModeChanged AppModeChangedDelegate;
diff --git a/Engine/Plugins/Animation/LiveLinkHub/Source/LiveLinkHub/Private/LiveLinkHubApplicationMode.cpp b/Engine/Plugins/Animation/LiveLinkHub/Source/LiveLinkHub/Private/LiveLinkHubApplicationMode.cpp
--- a/Engine/Plugins/Animation/LiveLinkHub/Source/LiveLinkHub/Private/LiveLinkHubApplicationMode.cpp
+++ b/Engine/Plugins/Animation/LiveLinkHub/Source/LiveLinkHub/Private/LiveLinkHubApplicationMode.cpp
@@ -2,6 +2,7 @@
 
 #include "LiveLinkHubApplicationMode.h"
 
+#include "LiveLinkHubApplicationBase.h"
 #include "UI/Widgets/SLiveLinkHubSettings.h"
 #include "WorkflowOrientedApp/WorkflowCentricApplication.h"
 #include "ToolMenus.h"
@@ -9,7 +10,15 @@
 const FName ILiveLinkHubApplicationModeFactory::ModularFeatureName = "LiveLinkHubApplicationModeFactory";
 const FName FLiveLinkHubApplicationMode::FileMenuExtensionPoint = "LiveLinkHubFileMenu";
 
-FLiveLinkHubApplicationMode::FLiveLinkHubApplicationMode(FName InApplicationMode, FText InDisplayName, TSharedPtr<FWorkflowCentricApplication> InApp)
+
+TSharedRef<class FLiveLinkHubApplicationMode> ILiveLinkHubApplicationModeFactory::CreateApplicationMode(TSharedPtr<FWorkflowCentricApplication> InApp)
+{
+	// Using MakeShareable here in order to call the private FLiveLinkHubApplicationMode constructor which is available to use since this class is a friend of it.
+	return MakeShareable(new FLiveLinkHubApplicationMode());
+};
+
+
+FLiveLinkHubApplicationMode::FLiveLinkHubApplicationMode(FName InApplicationMode, FText InDisplayName, TSharedPtr<FLiveLinkHubApplicationBase> InApp)
 	: FApplicationMode(InApplicationMode)
 {
 	LayoutIni = TEXT("LiveLinkHubLayout");
@@ -49,4 +58,5 @@ void FLiveLinkHubApplicationMode::PostActivateMode()
 void FLiveLinkHubApplicationMode::PreDeactivateMode()
 {
 	UToolMenus::Get()->UnregisterOwner(GetModeName());
-}
\ No newline at end of file
+}
+
diff --git a/Engine/Plugins/Animation/LiveLinkHub/Source/LiveLinkHub/Private/LiveLinkHubCreatorAppMode.h b/Engine/Plugins/Animation/LiveLinkHub/Source/LiveLinkHub/Private/LiveLinkHubCreatorAppMode.h
--- a/Engine/Plugins/Animation/LiveLinkHub/Source/LiveLinkHub/Private/LiveLinkHubCreatorAppMode.h
+++ b/Engine/Plugins/Animation/LiveLinkHub/Source/LiveLinkHub/Private/LiveLinkHubCreatorAppMode.h
@@ -8,6 +8,7 @@
 #include "Framework/MultiBox/MultiBoxBuilder.h"
 #include "IDetailsView.h"
 #include "ILiveLinkClient.h"
+#include "LiveLinkHubApplicationBase.h"
 #include "LiveLinkClientPanelViews.h"
 #include "LiveLinkClientPanelToolbar.h"
 #include "LiveLinkHub.h"
@@ -298,7 +299,7 @@ public:
 class FLiveLinkHubCreatorAppMode : public FLiveLinkHubApplicationMode
 {
 public:
-	FLiveLinkHubCreatorAppMode(TSharedPtr<class FWorkflowCentricApplication> App)
+	FLiveLinkHubCreatorAppMode(TSharedPtr<FLiveLinkHubApplicationBase> App)
 		: FLiveLinkHubApplicationMode("CreatorMode", LOCTEXT("CreatorModeLabel", "Creator Mode"), App)
 	{
 
@@ -424,17 +425,6 @@ public:
 			]
 		};
 	}
-
-	virtual void RegisterTabFactories(TSharedPtr<FTabManager> InTabManager) override
-	{
-		if (TSharedPtr<FWorkflowCentricApplication> App = WeakApp.Pin())
-		{
-			// Pushing the tab factories will register the tab spwaners for the current mode.
-			App->PushTabFactories(TabFactories);
-		}
-
-		FApplicationMode::RegisterTabFactories(InTabManager);
-	}
 	//~ End FLiveLinkHubApplicationMode interface
 
 private:
@@ -480,8 +470,6 @@ private:
 
 
 protected:
-	/** Set of spawnable tabs in the mode */
-	FWorkflowAllowedTabSet TabFactories;
 
 	/** Holds the livelink panel controller responsible for creating sources and subjects tabs. */
 	TSharedPtr<FLiveLinkPanelController> PanelController;
diff --git a/Engine/Plugins/Animation/LiveLinkHub/Source/LiveLinkHub/Public/LiveLinkHubApplicationBase.h b/Engine/Plugins/Animation/LiveLinkHub/Source/LiveLinkHub/Public/LiveLinkHubApplicationBase.h
new file mode 100644
--- /dev/null
+++ b/Engine/Plugins/Animation/LiveLinkHub/Source/LiveLinkHub/Public/LiveLinkHubApplicationBase.h
@@ -0,0 +1,23 @@
+// Copyright Epic Games, Inc. All Rights Reserved.
+
+#pragma once
+
+#include "UObject/NameTypes.h"
+#include "WorkflowOrientedApp/WorkflowCentricApplication.h"
+
+class FLiveLinkHubApplicationMode;
+
+/** Base class for the live link hub application, meant to expose a limited interface to modes that extend FLiveLinkHubApplicationMode. */
+class FLiveLinkHubApplicationBase : public FWorkflowCentricApplication
+{
+public:
+	virtual ~FLiveLinkHubApplicationBase() {}
+
+	/** Add a LiveLinkHub application mode. */
+	LIVELINKHUB_API virtual void AddLiveLinkHubApplicationMode(FName ModeName, TSharedRef<FLiveLinkHubApplicationMode> Mode) = 0;
+
+	/** Push tab factories for the given mode. */
+	LIVELINKHUB_API virtual void PushTabFactories(const class FWorkflowAllowedTabSet& FactorySetToPush, TSharedPtr<FLiveLinkHubApplicationMode> ApplicationMode) = 0;
+};
+
+
diff --git a/Engine/Plugins/Animation/LiveLinkHub/Source/LiveLinkHub/Public/LiveLinkHubApplicationMode.h b/Engine/Plugins/Animation/LiveLinkHub/Source/LiveLinkHub/Public/LiveLinkHubApplicationMode.h
--- a/Engine/Plugins/Animation/LiveLinkHub/Source/LiveLinkHub/Public/LiveLinkHubApplicationMode.h
+++ b/Engine/Plugins/Animation/LiveLinkHub/Source/LiveLinkHub/Public/LiveLinkHubApplicationMode.h
@@ -4,11 +4,14 @@
 
 #include "Features/IModularFeature.h"
 
+#include "LiveLinkHubApplicationBase.h"
 #include "WorkflowOrientedApp/ApplicationMode.h"
+#include "WorkflowOrientedApp/WorkflowTabManager.h"
 
 struct FSlateIcon;
 class FToolBarBuilder;
 class FWorkflowCentricApplication;
+class FLiveLinkHubApplicationBase;
 class SWidget;
 
 /** Factory used by external plugins to register custom LiveLinkHub layouts. */
@@ -18,8 +21,12 @@ public:
 	/** Name of the modular feature. */
 	static const FName ModularFeatureName;
 
+
+	UE_DEPRECATED(5.6, "Use CreateLiveLinkHubAppMode.")
+	virtual TSharedRef<class FLiveLinkHubApplicationMode> CreateApplicationMode(TSharedPtr<FWorkflowCentricApplication> InApp);
+
 	/** Instantiate an application mode so LiveLinkHub can register it and display it in its Layout Selector. */
-	virtual TSharedRef<class FLiveLinkHubApplicationMode> CreateApplicationMode(TSharedPtr<FWorkflowCentricApplication> InApp) = 0;
+	virtual TSharedRef<class FLiveLinkHubApplicationMode> CreateLiveLinkHubAppMode(TSharedPtr<FLiveLinkHubApplicationBase> InApp) = 0;
 };
 
 class LIVELINKHUB_API FLiveLinkHubApplicationMode : public FApplicationMode
@@ -29,7 +36,7 @@ public:
 	static const FName FileMenuExtensionPoint;
 
 	/** Get the icon for this mode. */
-	virtual FSlateIcon GetModeIcon() const = 0;
+	virtual FSlateIcon GetModeIcon() const { return FSlateIcon(); }
 
 	/** Get the display name for this mode */
 	const FText& GetDisplayName() const { return DisplayName; }
@@ -38,11 +45,25 @@ public:
 	TArray<TSharedRef<SWidget>> GetToolbarWidgets();
 
 protected:
-	FLiveLinkHubApplicationMode(FName InApplicationMode, FText InDisplayName, TSharedPtr<FWorkflowCentricApplication> InApp);
+	FLiveLinkHubApplicationMode(FName InApplicationMode, FText InDisplayName, TSharedPtr<FLiveLinkHubApplicationBase> InApp);
 
 	//~ Begin FApplicationMode interface
 	virtual void PostActivateMode() override;
 	virtual void PreDeactivateMode() override;
+	virtual void RegisterTabFactories(TSharedPtr<FTabManager> InTabManager) override
+	{
+		if (!bRegisteredTabFactories)
+		{
+			if (TSharedPtr<FLiveLinkHubApplicationBase> App = WeakApp.Pin())
+			{
+				// Pushing the tab factories will register the tab spwaners for the current mode.
+				App->PushTabFactories(TabFactories, StaticCastSharedRef<FLiveLinkHubApplicationMode>(AsShared()));
+			}
+
+			FApplicationMode::RegisterTabFactories(InTabManager);
+			bRegisteredTabFactories = true;
+		}
+	}
 	//~ End FApplicationMode interface
 
 	/** Implement in a child class in order to have status bar extensions show up when the mode is active. */
@@ -51,12 +72,25 @@ protected:
 	
 	/** Get widgets that should appear in the status bar. */
 	virtual TArray<TSharedRef<SWidget>> GetStatusBarWidgets_Impl() { return {}; }
-
 protected:
 	/** The application that hosts this app mode. */
-	TWeakPtr<FWorkflowCentricApplication> WeakApp;
+	TWeakPtr<FLiveLinkHubApplicationBase> WeakApp;
+
+	/** Set of spawnable tabs in the mode */
+	FWorkflowAllowedTabSet TabFactories;
+
+private:
+	friend class ILiveLinkHubApplicationModeFactory;
+	// Invalid contructor meant to be used by ILiveLinkHubApplicationModeFactory to support the deprecated factory method.
+	FLiveLinkHubApplicationMode()
+		: FApplicationMode("INVALID_MODE")
+	{
+	}
 
 private:
 	/** Display name for this mode. */
 	FText DisplayName;
+
+	/** Keeps track of whether or not the tab factories were registered, so we don't end up with duplicate tab spawners for this mode. */
+	bool bRegisteredTabFactories = false;
 };
\ No newline at end of file
diff --git a/Engine/Plugins/VirtualProduction/CaptureManager/CaptureManagerApp/Source/CaptureManagerEditor/Private/CaptureManagerEditorModule.cpp b/Engine/Plugins/VirtualProduction/CaptureManager/CaptureManagerApp/Source/CaptureManagerEditor/Private/CaptureManagerEditorModule.cpp
--- a/Engine/Plugins/VirtualProduction/CaptureManager/CaptureManagerApp/Source/CaptureManagerEditor/Private/CaptureManagerEditorModule.cpp
+++ b/Engine/Plugins/VirtualProduction/CaptureManager/CaptureManagerApp/Source/CaptureManagerEditor/Private/CaptureManagerEditorModule.cpp
@@ -3,6 +3,7 @@
 #include "CaptureManagerEditorModule.h"
 
 #include "Features/IModularFeatures.h"
+#include "LiveLinkHubApplicationBase.h"
 #include "UI/LiveLinkHubCaptureManagerMode.h"
 #include "ThumbnailRendering/ThumbnailManager.h"
 #include "UI/TakeThumbnailRenderer.h"
@@ -36,7 +37,7 @@ void FCaptureManagerEditorModule::ShutdownModule()
 }
 
 /** Instantiate an application mode so LiveLinkHub can register it and display it in its Layout Selector. */
-TSharedRef<FLiveLinkHubApplicationMode> FCaptureManagerEditorModule::CreateApplicationMode(TSharedPtr<FWorkflowCentricApplication> InApp)
+TSharedRef<FLiveLinkHubApplicationMode> FCaptureManagerEditorModule::CreateLiveLinkHubAppMode(TSharedPtr<FLiveLinkHubApplicationBase> InApp)
 {
 	return MakeShared<FLiveLinkHubCaptureManagerMode>(InApp);
 }
\ No newline at end of file
diff --git a/Engine/Plugins/VirtualProduction/CaptureManager/CaptureManagerApp/Source/CaptureManagerEditor/Private/UI/LiveLinkHubCaptureManagerMode.cpp b/Engine/Plugins/VirtualProduction/CaptureManager/CaptureManagerApp/Source/CaptureManagerEditor/Private/UI/LiveLinkHubCaptureManagerMode.cpp
--- a/Engine/Plugins/VirtualProduction/CaptureManager/CaptureManagerApp/Source/CaptureManagerEditor/Private/UI/LiveLinkHubCaptureManagerMode.cpp
+++ b/Engine/Plugins/VirtualProduction/CaptureManager/CaptureManagerApp/Source/CaptureManagerEditor/Private/UI/LiveLinkHubCaptureManagerMode.cpp
@@ -6,6 +6,7 @@
 #include "CaptureCommunicationManagerModule.h"
 #include "CaptureManagerPanelController.h"
 #include "CaptureUtilsModule.h"
+#include "LiveLinkHubApplicationBase.h"
 #include "LiveLinkHubDiscoveryModule.h"
 #include "Network/NetworkMisc.h"
 #include "SIngestJobProcessor.h"
@@ -53,7 +54,7 @@ namespace UE::CaptureManager
 struct FTakeBrowserTabSummoner : public FWorkflowTabFactory
 {
 public:
-	FTakeBrowserTabSummoner(TSharedPtr<FAssetEditorToolkit> InHostingApp, TSharedPtr<FCaptureManagerPanelController> InPanelController)
+	FTakeBrowserTabSummoner(TSharedPtr<FLiveLinkHubApplicationBase> InHostingApp, TSharedPtr<FCaptureManagerPanelController> InPanelController)
 		: FWorkflowTabFactory(UE::CaptureManager::TakeBrowserTabId, MoveTemp(InHostingApp))
 		, PanelController(MoveTemp(InPanelController))
 	{
@@ -75,7 +76,7 @@ private:
 struct FJobsListTabSummoner : public FWorkflowTabFactory
 {
 public:
-	FJobsListTabSummoner(TSharedPtr<FAssetEditorToolkit> InHostingApp, TSharedPtr<FCaptureManagerPanelController> InPanelController)
+	FJobsListTabSummoner(TSharedPtr<FLiveLinkHubApplicationBase> InHostingApp, TSharedPtr<FCaptureManagerPanelController> InPanelController)
 		: FWorkflowTabFactory(UE::CaptureManager::JobsListTabId, MoveTemp(InHostingApp))
 		, PanelController(MoveTemp(InPanelController))
 	{
@@ -101,7 +102,7 @@ private:
 struct FJobsDetailsTabSummoner : public FWorkflowTabFactory
 {
 public:
-	FJobsDetailsTabSummoner(TSharedPtr<FAssetEditorToolkit> InHostingApp, TSharedPtr<FCaptureManagerPanelController> InPanelController)
+	FJobsDetailsTabSummoner(TSharedPtr<FLiveLinkHubApplicationBase> InHostingApp, TSharedPtr<FCaptureManagerPanelController> InPanelController)
 		: FWorkflowTabFactory(UE::CaptureManager::JobDetailsTabId, MoveTemp(InHostingApp))
 		, PanelController(MoveTemp(InPanelController))
 	{
@@ -122,7 +123,7 @@ private:
 	TSharedPtr<FCaptureManagerPanelController> PanelController;
 };
 
-FLiveLinkHubCaptureManagerMode::FLiveLinkHubCaptureManagerMode(TSharedPtr<class FWorkflowCentricApplication> App)
+FLiveLinkHubCaptureManagerMode::FLiveLinkHubCaptureManagerMode(TSharedPtr<FLiveLinkHubApplicationBase> App)
 	: FLiveLinkHubApplicationMode("CaptureManager", LOCTEXT("CaptureManagerModeLabel", "Capture Manager"), App)
 	, TimerManager(UE::CaptureManager::GetTimerManager())
 	, DiscoveryRequester(UE::CaptureManager::GetDiscoveryRequester())
@@ -213,18 +214,6 @@ TArray<TSharedRef<SWidget>> FLiveLinkHubCaptureManagerMode::GetStatusBarWidgets_
 	};
 }
 
-
-void FLiveLinkHubCaptureManagerMode::RegisterTabFactories(TSharedPtr<FTabManager> InTabManager)
-{
-	if (TSharedPtr<FWorkflowCentricApplication> App = WeakApp.Pin())
-	{
-		// Pushing the tab factories will register the tab spwaners for the current mode.
-		App->PushTabFactories(TabFactories);
-	}
-
-	FApplicationMode::RegisterTabFactories(InTabManager);
-}
-
 void FLiveLinkHubCaptureManagerMode::OnTick(TSharedRef<UE::CaptureManager::FDiscoveryRequester> InDiscoveryRequester,
 											TSharedRef<FCommunicationManager> InCommunicationManager)
 {
diff --git a/Engine/Plugins/VirtualProduction/CaptureManager/CaptureManagerApp/Source/CaptureManagerEditor/Private/UI/LiveLinkHubCaptureManagerMode.h b/Engine/Plugins/VirtualProduction/CaptureManager/CaptureManagerApp/Source/CaptureManagerEditor/Private/UI/LiveLinkHubCaptureManagerMode.h
--- a/Engine/Plugins/VirtualProduction/CaptureManager/CaptureManagerApp/Source/CaptureManagerEditor/Private/UI/LiveLinkHubCaptureManagerMode.h
+++ b/Engine/Plugins/VirtualProduction/CaptureManager/CaptureManagerApp/Source/CaptureManagerEditor/Private/UI/LiveLinkHubCaptureManagerMode.h
@@ -27,7 +27,7 @@ class SDockTab;
 class FLiveLinkHubCaptureManagerMode : public FLiveLinkHubApplicationMode
 {
 public:
-	FLiveLinkHubCaptureManagerMode(TSharedPtr<class FWorkflowCentricApplication> App);
+	FLiveLinkHubCaptureManagerMode(TSharedPtr<FLiveLinkHubApplicationBase> App);
 	virtual ~FLiveLinkHubCaptureManagerMode();
 
 	//~ Begin FLiveLinkHubApplicationMode interface
@@ -36,15 +36,8 @@ public:
 		return FSlateIcon(FCaptureManagerTabStyle::Get().GetStyleSetName(), TEXT("CaptureManager.Tabs.CaptureManager"));
 	}
 
-	/** Get widgets that should appear in the status bar. */
 	virtual TArray<TSharedRef<SWidget>> GetStatusBarWidgets_Impl() override;
-
 	//~ End FLiveLinkHubApplicationMode interface
-
-	//~ Begin FApplicationMode interface
-	virtual void RegisterTabFactories(TSharedPtr<FTabManager> InTabManager) override;
-	//~ End FApplicationMode interace
-
 private:
 
 	static constexpr float ConnectionInvokePeriod = 3.0f;
@@ -67,9 +60,6 @@ private:
 	FDelegateHandle ConnectionChangedHandler;
 
 	TSharedPtr<class FCaptureManagerPanelController> PanelController;
-
-	/** Set of spawnable tabs in the mode */
-	FWorkflowAllowedTabSet TabFactories;
 	
 	std::atomic<bool> bIsConnected = false;
 };
diff --git a/Engine/Plugins/VirtualProduction/CaptureManager/CaptureManagerApp/Source/CaptureManagerEditor/Public/CaptureManagerEditorModule.h b/Engine/Plugins/VirtualProduction/CaptureManager/CaptureManagerApp/Source/CaptureManagerEditor/Public/CaptureManagerEditorModule.h
--- a/Engine/Plugins/VirtualProduction/CaptureManager/CaptureManagerApp/Source/CaptureManagerEditor/Public/CaptureManagerEditorModule.h
+++ b/Engine/Plugins/VirtualProduction/CaptureManager/CaptureManagerApp/Source/CaptureManagerEditor/Public/CaptureManagerEditorModule.h
@@ -14,6 +14,6 @@ public:
 	virtual void ShutdownModule() override;
 
 	//~ Begin ILiveLinkHubApplicationModeFactory
-	virtual TSharedRef<FLiveLinkHubApplicationMode> CreateApplicationMode(TSharedPtr<FWorkflowCentricApplication> InApp) override;
+	virtual TSharedRef<FLiveLinkHubApplicationMode> CreateLiveLinkHubAppMode(TSharedPtr<class FLiveLinkHubApplicationBase> InApp) override;
 	//~ End ILiveLinkHubApplicationModeFactory
 };
