Add RAM usage display to LiveLinkHub

#rnx

diff --git a/Engine/Plugins/Animation/LiveLinkHub/Source/LiveLinkHub/LiveLinkHub.Build.cs b/Engine/Plugins/Animation/LiveLinkHub/Source/LiveLinkHub/LiveLinkHub.Build.cs
--- a/Engine/Plugins/Animation/LiveLinkHub/Source/LiveLinkHub/LiveLinkHub.Build.cs
+++ b/Engine/Plugins/Animation/LiveLinkHub/Source/LiveLinkHub/LiveLinkHub.Build.cs
@@ -44,6 +44,7 @@ public class LiveLinkHub : ModuleRules
 				"TimeManagement",
 				"ToolMenus",
 				"ToolWidgets",
+				"TraceLog",
 				"UnrealEd",
 				"OutputLog",
 				"NamingTokens",
diff --git a/Engine/Plugins/Animation/LiveLinkHub/Source/LiveLinkHub/Private/Settings/LiveLinkHubSettings.h b/Engine/Plugins/Animation/LiveLinkHub/Source/LiveLinkHub/Private/Settings/LiveLinkHubSettings.h
--- a/Engine/Plugins/Animation/LiveLinkHub/Source/LiveLinkHub/Private/Settings/LiveLinkHubSettings.h
+++ b/Engine/Plugins/Animation/LiveLinkHub/Source/LiveLinkHub/Private/Settings/LiveLinkHubSettings.h
@@ -56,9 +56,17 @@ public:
 	TArray<FName> ProjectSettingsToDisplay;
 
 	/** If this is enabled, invalid subjects will be removed after loading a session. */
-	UPROPERTY(config, EditAnywhere, AdvancedDisplay, Category = "LiveLinkHub")
+	UPROPERTY(config, EditAnywhere, Category = "LiveLinkHub")
 	bool bRemoveInvalidSubjectsAfterLoadingSession = false;
 
+	/** Whether to show memory usage in the top right corner. */
+	UPROPERTY(config, EditAnywhere, Category = "LiveLinkHub")
+	bool bShowMemoryUsage = true;
+
+	/** How much RAM (in MB) the program can use before showing a warning. */
+	UPROPERTY(config, EditAnywhere, Category = "LiveLinkHub", DisplayName = "Memory Warning Threshold (MB)")
+	float ShowMemoryWarningThresholdMB = 5000.0;
+
 	/**
 	 * - Experimental - If this is disabled, LiveLinkHub's LiveLink Client will tick outside of the game thread.
 	 * This allows processing LiveLink frame snapshots without the risk of being blocked by the game / ui thread.
diff --git a/Engine/Plugins/Animation/LiveLinkHub/Source/LiveLinkHub/Private/UI/Widgets/SLiveLinkHubMemoryStats.h b/Engine/Plugins/Animation/LiveLinkHub/Source/LiveLinkHub/Private/UI/Widgets/SLiveLinkHubMemoryStats.h
new file mode 100644
--- /dev/null
+++ b/Engine/Plugins/Animation/LiveLinkHub/Source/LiveLinkHub/Private/UI/Widgets/SLiveLinkHubMemoryStats.h
@@ -0,0 +1,212 @@
+// Copyright Epic Games, Inc. All Rights Reserved.
+
+#pragma once
+
+#include "Widgets/SCompoundWidget.h"
+
+#include "Framework/Notifications/NotificationManager.h"
+#include "LiveLinkHubLog.h"
+#include "Misc/Paths.h"
+#include "Settings/LiveLinkHubSettings.h"
+#include "Widgets/Images/SImage.h"
+#include "Widgets/Input/SButton.h"
+#include "Widgets/Text/STextBlock.h"
+#include "Widgets/Notifications/SNotificationList.h"
+
+
+#define LOCTEXT_NAMESPACE "LiveLinkHubMemoryStats"
+
+/** Widget that displays the current memory usage of the program and shows a warning if the memory usage falls above a threshold defined in the LiveLinkHub settings. */
+class SLiveLinkHubMemoryStats : public SCompoundWidget
+{
+public:
+	SLATE_BEGIN_ARGS(SLiveLinkHubMemoryStats)
+		{}
+	SLATE_END_ARGS()
+
+	void Construct(const FArguments& InArgs)
+	{
+		ChildSlot
+		[
+			SNew(SHorizontalBox)
+			.Visibility_Lambda([]() { return GetDefault<ULiveLinkHubSettings>()->bShowMemoryUsage ? EVisibility::Visible : EVisibility::Collapsed; })
+			+ SHorizontalBox::Slot()
+			.VAlign(VAlign_Center)
+			.AutoWidth()
+			[
+				SNew(STextBlock)
+				.Text(LOCTEXT("MemoryLabel", "Mem: "))
+			]
+			+ SHorizontalBox::Slot()
+			.VAlign(VAlign_Center)
+			.AutoWidth()
+			[
+				SNew(STextBlock)
+				.Text(TAttribute<FText>::CreateLambda([this]() { return CachedMemoryText; }))
+			] 
+			+ SHorizontalBox::Slot()
+			.AutoWidth()
+			.Padding(2.0f, 0.0f)
+			.VAlign(VAlign_Center)
+			[
+				SNew(SButton)
+				.ContentPadding(FMargin(0, 2.0))
+				.ToolTipText(LOCTEXT("MemoryWarningTooltip", "Live Link Hub is using an unusual amount of RAM. Click here to collect a memory trace."))
+				.Visibility(this, &SLiveLinkHubMemoryStats::OnGetWarningVisibility)
+				.OnClicked(this, &SLiveLinkHubMemoryStats::OnClickMemoryWarning)
+				.IsEnabled_Lambda([this]() { return !bCreatingTrace; })
+				[
+					SNew(SImage)
+					.ColorAndOpacity(FLinearColor::Red)
+					.Image(FCoreStyle::Get().GetBrush("Icons.Warning"))
+				]
+			]
+		];
+	}
+
+	virtual void Tick(const FGeometry & AllottedGeometry, const double InCurrentTime, const float InDeltaTime) override
+	{
+		UpdateMemoryStats();
+
+		if (CachedTotalPhysicalUsedMB < GetDefault<ULiveLinkHubSettings>()->ShowMemoryWarningThresholdMB)
+		{
+			bWarnedUser = false;
+		}
+		else if (!bWarnedUser && !CurrentNotification)
+		{
+			WarnUser();
+		}
+	}
+
+private:
+	/** Update our cached memory stats. */
+	void UpdateMemoryStats()
+	{
+		// Only refresh process memory allocated after every so often, to reduce fixed frame time overhead
+		static SIZE_T StaticLastTotalAllocated = 0;
+		static int32 QueriesUntilUpdate = 1;
+		if (--QueriesUntilUpdate <= 0)
+		{
+			// Query OS for process memory used
+			FPlatformMemoryStats MemoryStats = FPlatformMemory::GetStats();
+			StaticLastTotalAllocated = MemoryStats.UsedPhysical;
+
+			// Wait 120 queries until we refresh memory again
+			QueriesUntilUpdate = FramesBetweenPlatformQueries;
+		}
+
+		static const FNumberFormattingOptions FormatOptions = FNumberFormattingOptions()
+			.SetMinimumFractionalDigits(2)
+			.SetMaximumFractionalDigits(2);
+		static const FText MemorySizeFmt = FText::FromString(TEXT("{0} mb"));
+
+		CachedTotalPhysicalUsedMB = (float)StaticLastTotalAllocated / (1024.0f * 1024.0f);
+		CachedMemoryText = FText::Format(MemorySizeFmt, FText::AsNumber(CachedTotalPhysicalUsedMB, &FormatOptions));
+
+		if (CurrentNotification && CurrentNotification->GetCompletionState() == SNotificationItem::CS_None)
+		{
+			CurrentNotification->SetText(CreateWarningText());
+		}
+	}
+
+	/** Open a toast notification to warn the user that memory consumption is unusually high. */
+	void WarnUser()
+	{
+		FNotificationInfo Info(CreateWarningText());
+		Info.Image = FAppStyle::GetBrush("Icons.WarningWithColor");
+		Info.bFireAndForget = false;
+		Info.WidthOverride = 500;
+
+		FNotificationButtonInfo CreateReportButton{ LOCTEXT("CreateTraceButton", "Create trace snapshot"), LOCTEXT("CreateTraceButtonToolTip", "Generate a trace snapshot."), FSimpleDelegate::CreateSP(this, &SLiveLinkHubMemoryStats::GenerateTraceSnapshot), SNotificationItem::ECompletionState::CS_None };
+		FNotificationButtonInfo CloseNotificationButton{ LOCTEXT("CloseNotificationButton", "Close"), FText::GetEmpty(), FSimpleDelegate::CreateSP(this, &SLiveLinkHubMemoryStats::CloseNotification), SNotificationItem::ECompletionState::CS_None };
+		Info.ButtonDetails.Add(CreateReportButton);
+		Info.ButtonDetails.Add(CloseNotificationButton);
+
+		CurrentNotification = FSlateNotificationManager::Get().AddNotification(Info);
+		CurrentNotification->SetCompletionState(SNotificationItem::CS_None);
+
+		bWarnedUser = true;
+	}
+
+	/** Handles clicking on the warning button to create a trace snapshot. */
+	FReply OnClickMemoryWarning()
+	{
+		GenerateTraceSnapshot();
+
+		return FReply::Handled();
+	}
+
+	/** Handles creating a snapshot tracefile. */
+	void GenerateTraceSnapshot()
+	{
+		if (bCreatingTrace)
+		{
+			return;
+		}
+
+		bCreatingTrace = true;
+
+		if (CurrentNotification)
+		{
+			CurrentNotification->SetText(LOCTEXT("CreatedTrace", "Created trace snapshot..."));
+			CurrentNotification->SetCompletionState(SNotificationItem::CS_Success);
+			CurrentNotification->ExpireAndFadeout();
+		}
+
+		const FString TraceFileName = FDateTime::Now().ToString(TEXT("LiveLinkHubMemoryLeak - %Y%m%d_%H%M%S.utrace"));
+		LastTracePath = FPaths::Combine(FPaths::ProfilingDir(), TraceFileName);
+
+		if (!UE::Trace::WriteSnapshotTo(*LastTracePath))
+		{
+			UE_LOG(LogLiveLinkHub, Warning, TEXT("Failed writing trace snapshot file."));
+		}
+		else
+		{
+			UE_LOG(LogLiveLinkHub, Display, TEXT("Wrote trace snapshot file to %s"), *LastTracePath);
+		}
+
+
+		bCreatingTrace = false;
+	}
+
+	/** Closes the current memory usage notification. */
+	void CloseNotification()
+	{
+		CurrentNotification->Fadeout();
+		CurrentNotification.Reset();
+	}
+
+	/** Determines whether the warning icon should be visible. */
+	EVisibility OnGetWarningVisibility() const
+	{
+		if (CachedTotalPhysicalUsedMB > GetDefault<ULiveLinkHubSettings>()->ShowMemoryWarningThresholdMB)
+		{
+			return EVisibility::Visible;
+		}
+		return EVisibility::Collapsed;
+	}
+
+	/** Generate the warning message for current memory usage. */
+	FText CreateWarningText()
+	{
+		const FText WarningMessage = FText::Format(LOCTEXT("PhysicalMemoryWarning", "Physical memory usage is unusually high! ({0})."), CachedMemoryText);
+		return WarningMessage;
+	}
+
+	/** Cached memory used in this process. */
+	float CachedTotalPhysicalUsedMB = 0.f;
+	/** Cached text representation of our total physical memory used. */
+	FText CachedMemoryText;
+	/** Stores the path to the last trace snapshot that was generated. */
+	FString LastTracePath;
+	/** Whether the hub is currently generating a report. */
+	bool bCreatingTrace = false;
+	/** Whether the user was warned about high memory usage (will reset once it goes below threshold again. */
+	bool bWarnedUser = false;
+	/** Pointer to the current notification shown on screen. */
+	TSharedPtr<SNotificationItem> CurrentNotification;
+	/** Number of frames to wait between calls to FPlatformMemory::GetStats */
+	const int32 FramesBetweenPlatformQueries = 120;
+};
+
+#undef LOCTEXT_NAMESPACE
\ No newline at end of file
diff --git a/Engine/Plugins/Animation/LiveLinkHub/Source/LiveLinkHub/Private/UI/Window/LiveLinkHubWindowController.cpp b/Engine/Plugins/Animation/LiveLinkHub/Source/LiveLinkHub/Private/UI/Window/LiveLinkHubWindowController.cpp
--- a/Engine/Plugins/Animation/LiveLinkHub/Source/LiveLinkHub/Private/UI/Window/LiveLinkHubWindowController.cpp
+++ b/Engine/Plugins/Animation/LiveLinkHub/Source/LiveLinkHub/Private/UI/Window/LiveLinkHubWindowController.cpp
@@ -12,6 +12,7 @@
 #include "HAL/PlatformApplicationMisc.h"
 #include "Misc/ConfigCacheIni.h"
 #include "UI/Widgets/LiveLinkHubMainTabController.h"
+#include "UI/Widgets/SLiveLinkHubMemoryStats.h"
 #include "UI/Widgets/SLiveLinkHubTopologyModeSwitcher.h"
 #include "UI/Window/ModalWindowManager.h"
 #include "Widgets/Docking/SDockTab.h"
@@ -110,7 +111,7 @@ void FLiveLinkHubWindowController::RestoreLayout(TSharedPtr<FAssetEditorToolkit>
 				.Text(FText::FromString(FLiveLinkHub::Get()->GetLiveLinkProvider()->GetProviderName()))
 		];
 
-	TSharedRef<SWidget> RightMenuWidget = 
+	TSharedRef<SWidget> RightMenuWidget =
 		SNew(SHorizontalBox)
 		+ SHorizontalBox::Slot()
 		.AutoWidth()
@@ -122,6 +123,12 @@ void FLiveLinkHubWindowController::RestoreLayout(TSharedPtr<FAssetEditorToolkit>
 		.Padding(4.0f)
 		[
 			SNew(SLiveLinkHubTopologyModeSwitcher)
+		]
+		+ SHorizontalBox::Slot()
+		.AutoWidth()
+		.Padding(4.0f)
+		[
+			SNew(SLiveLinkHubMemoryStats)
 		];
 
 		FGlobalTabmanager::Get()->FindExistingLiveTab(StandaloneTabId)->SetTitleBarRightContent(RightMenuWidget);
