NamingTokens: Remove FolderPath and FileName from ResultData, as this should be up to the caller to calculate.

diff --git a/Engine/Plugins/Animation/LiveLinkHub/Source/LiveLinkHub/Private/Recording/Implementations/LiveLinkUAssetRecorder.cpp b/Engine/Plugins/Animation/LiveLinkHub/Source/LiveLinkHub/Private/Recording/Implementations/LiveLinkUAssetRecorder.cpp
--- a/Engine/Plugins/Animation/LiveLinkHub/Source/LiveLinkHub/Private/Recording/Implementations/LiveLinkUAssetRecorder.cpp
+++ b/Engine/Plugins/Animation/LiveLinkHub/Source/LiveLinkHub/Private/Recording/Implementations/LiveLinkUAssetRecorder.cpp
@@ -164,12 +164,30 @@ bool FLiveLinkUAssetRecorder::GetSavePresetPackageName(FString& OutName)
 	FNamingTokenResultData TemplateData = GEditor->GetEditorSubsystem<UNamingTokensEditorSubsystem>()->EvaluateTokenString(GetDefault<ULiveLinkHubSettings>()->FilenameTemplate,
 		Filter);
 
-	const FString DefaultName = TemplateData.FileName;
-	const FString DefaultFolder = TemplateData.FolderPath;
+	FString DefaultName;
+	FString DefaultFolder;
+	// Split folder path and file name.
+	{
+		FString FormattedString = TemplateData.EvaluatedText.ToString();
+		int32 LastSlashIndex;
+		if (FormattedString.FindLastChar('/', LastSlashIndex))
+		{
+			DefaultFolder = FormattedString.Left(LastSlashIndex);
+			if (DefaultFolder.StartsWith("/"))
+			{
+				DefaultFolder.RemoveFromStart("/");
+			}
+			DefaultName = FormattedString.Mid(LastSlashIndex + 1);
+		}
+		else
+		{
+			DefaultName = FormattedString;
+		}
+	}
 	
 	const FString ContentDir = FPaths::ProjectContentDir();
-	const FString DialogStartPath = FPaths::Combine(TEXT("/Game"), TemplateData.FolderPath);
-	const FString AbsoluteFolderPath = FPaths::Combine(ContentDir, TemplateData.FolderPath);
+	const FString DialogStartPath = FPaths::Combine(TEXT("/Game"), DefaultFolder);
+	const FString AbsoluteFolderPath = FPaths::Combine(ContentDir, DefaultFolder);
 
 	IPlatformFile& PlatformFile = FPlatformFileManager::Get().GetPlatformFile();
 
diff --git a/Engine/Plugins/Developer/NamingTokens/Source/NamingTokens/Private/NamingTokens.cpp b/Engine/Plugins/Developer/NamingTokens/Source/NamingTokens/Private/NamingTokens.cpp
--- a/Engine/Plugins/Developer/NamingTokens/Source/NamingTokens/Private/NamingTokens.cpp
+++ b/Engine/Plugins/Developer/NamingTokens/Source/NamingTokens/Private/NamingTokens.cpp
@@ -131,24 +131,6 @@ FNamingTokenResultData UNamingTokens::EvaluateTokenText(const FText& InTokenText
 	}
 
 	Result.EvaluatedText = FormattedText;
-	
-	// Split folder path and file name.
-	FString FormattedString = FormattedText.ToString();
-	int32 LastSlashIndex;
-	if (FormattedString.FindLastChar('/', LastSlashIndex))
-	{
-		Result.FolderPath = FormattedString.Left(LastSlashIndex);
-		if (Result.FolderPath.StartsWith("/"))
-		{
-			Result.FolderPath.RemoveFromStart("/");
-		}
-		Result.FileName = FormattedString.Mid(LastSlashIndex + 1);
-	}
-	else
-	{
-		Result.FolderPath = TEXT("");
-		Result.FileName = FormattedString;
-	}
 
 	PostEvaluate(Result);
 	
diff --git a/Engine/Plugins/Developer/NamingTokens/Source/NamingTokens/Public/NamingTokenData.h b/Engine/Plugins/Developer/NamingTokens/Source/NamingTokens/Public/NamingTokenData.h
--- a/Engine/Plugins/Developer/NamingTokens/Source/NamingTokens/Public/NamingTokenData.h
+++ b/Engine/Plugins/Developer/NamingTokens/Source/NamingTokens/Public/NamingTokenData.h
@@ -40,15 +40,6 @@ struct NAMINGTOKENS_API FNamingTokenResultData
 	/** The full text with evaluated tokens. */
 	UPROPERTY(BlueprintReadOnly, Category = "NamingTokens")
 	FText EvaluatedText;
-	/** @todo NamingTokens - Remove and leave up to caller to manage.
-	 * The evaluated folder(s) without the first forward slash or filename. */
-	UPROPERTY(BlueprintReadOnly, Category = "NamingTokens")
-	FString FolderPath;
-	/** @todo NamingTokens - Remove and leave up to caller to manage.
-	 * The evaluated filename. */
-	UPROPERTY(BlueprintReadOnly, Category = "NamingTokens")
-	FString FileName;
-
 	/** The result of individual tokens, in the order they appear in OriginalText. */
 	UPROPERTY(BlueprintReadOnly, Category = "NamingTokens")
 	TArray<FNamingTokenValueData> TokenValues;
