LiveLinkHub: Connect Session, Slate input boxes to NamingTokens. The old session token has been renamed to config.

Take will auto increment on a successful recording.

Disable record button if session and slate aren't properly filled out.

diff --git a/Engine/Plugins/Animation/LiveLinkHub/Source/LiveLinkHub/Private/Config/LiveLinkHubTemplateTokens.cpp b/Engine/Plugins/Animation/LiveLinkHub/Source/LiveLinkHub/Private/Config/LiveLinkHubTemplateTokens.cpp
--- a/Engine/Plugins/Animation/LiveLinkHub/Source/LiveLinkHub/Private/Config/LiveLinkHubTemplateTokens.cpp
+++ b/Engine/Plugins/Animation/LiveLinkHub/Source/LiveLinkHub/Private/Config/LiveLinkHubTemplateTokens.cpp
@@ -18,11 +18,38 @@ void ULiveLinkHubNamingTokens::OnCreateDefaultTokens(TArray<FNamingTokenData>& T
 
 	Tokens.Add({
 		TEXT("session"),
-		LOCTEXT("SessionName", "Session Name"),
+		LOCTEXT("SessionName", "Session Name from the Session Entry Widget"),
 		FNamingTokenData::FTokenProcessorDelegateNative::CreateLambda([this]() {
 			return FText::FromString(FString::Printf(TEXT("%s"), *SessionName));
 		})
 	});
+
+	Tokens.Add({
+		TEXT("slate"),
+		LOCTEXT("SlateName", "Slate Name from the Session Entry Widget"),
+		FNamingTokenData::FTokenProcessorDelegateNative::CreateLambda([this]() {
+			return FText::FromString(FString::Printf(TEXT("%s"), *SlateName));
+		})
+	});
+
+	Tokens.Add({
+		TEXT("take"),
+		LOCTEXT("TakeNumber", "Take Number from the Session Entry Widget"),
+		FNamingTokenData::FTokenProcessorDelegateNative::CreateLambda([this]() {
+			FNumberFormattingOptions NumberFormatOptions;
+			NumberFormatOptions.MinimumIntegralDigits = 2;
+
+			return FText::AsNumber(TakeNumber, &NumberFormatOptions);
+		})
+	});
+	
+	Tokens.Add({
+		TEXT("config"),
+		LOCTEXT("ConfigName", "Loaded Live Link Hub Config Name"),
+		FNamingTokenData::FTokenProcessorDelegateNative::CreateLambda([this]() {
+			return FText::FromString(FString::Printf(TEXT("%s"), *ConfigName));
+		})
+	});
 }
 
 void ULiveLinkHubNamingTokens::OnPreEvaluate_Implementation(const FNamingTokensEvaluationData& InEvaluationData)
@@ -34,7 +61,12 @@ void ULiveLinkHubNamingTokens::OnPreEvaluate_Implementation(const FNamingTokensE
 		{
 			if (const TSharedPtr<ILiveLinkHubSessionManager> SessionManager = LiveLinkHub->GetSessionManager())
 			{
-				SessionName = FPaths::GetBaseFilename(SessionManager->GetLastConfigPath());
+				ConfigName = FPaths::GetBaseFilename(SessionManager->GetLastConfigPath());
+				
+				const ILiveLinkRecordingSessionInfo& SessionInfo = ILiveLinkRecordingSessionInfo::Get();
+				SessionName = SessionInfo.GetSessionName();
+				TakeNumber = SessionInfo.GetTakeNumber();
+				SlateName = SessionInfo.GetSlateName();
 			}
 		}
 	}
@@ -43,7 +75,10 @@ void ULiveLinkHubNamingTokens::OnPreEvaluate_Implementation(const FNamingTokensE
 void ULiveLinkHubNamingTokens::OnPostEvaluate_Implementation()
 {
 	Super::OnPostEvaluate_Implementation();
+	ConfigName.Empty();
 	SessionName.Empty();
+	SlateName.Empty();
+	TakeNumber = 0;
 }
 
 #undef LOCTEXT_NAMESPACE
diff --git a/Engine/Plugins/Animation/LiveLinkHub/Source/LiveLinkHub/Private/Config/LiveLinkHubTemplateTokens.h b/Engine/Plugins/Animation/LiveLinkHub/Source/LiveLinkHub/Private/Config/LiveLinkHubTemplateTokens.h
--- a/Engine/Plugins/Animation/LiveLinkHub/Source/LiveLinkHub/Private/Config/LiveLinkHubTemplateTokens.h
+++ b/Engine/Plugins/Animation/LiveLinkHub/Source/LiveLinkHub/Private/Config/LiveLinkHubTemplateTokens.h
@@ -22,6 +22,15 @@ protected:
 	// ~End UNamingTokens
 	
 private:
+	/** The loaded config name. */
+	FString ConfigName;
+	
 	/** The current session. */
 	FString SessionName;
+
+	/** Session slate name */
+	FString SlateName;
+	
+	/** Session take number. */
+	int32 TakeNumber = 0;
 };
diff --git a/Engine/Plugins/Animation/LiveLinkHub/Source/LiveLinkHub/Private/Recording/Implementations/LiveLinkUAssetRecorder.cpp b/Engine/Plugins/Animation/LiveLinkHub/Source/LiveLinkHub/Private/Recording/Implementations/LiveLinkUAssetRecorder.cpp
--- a/Engine/Plugins/Animation/LiveLinkHub/Source/LiveLinkHub/Private/Recording/Implementations/LiveLinkUAssetRecorder.cpp
+++ b/Engine/Plugins/Animation/LiveLinkHub/Source/LiveLinkHub/Private/Recording/Implementations/LiveLinkUAssetRecorder.cpp
@@ -12,6 +12,7 @@
 #include "Framework/MultiBox/MultiBoxBuilder.h"
 #include "HAL/PlatformFileManager.h"
 #include "IContentBrowserSingleton.h"
+#include "ILiveLinkRecordingSessionInfo.h"
 #include "LiveLinkHub.h"
 #include "LiveLinkHubClient.h"
 #include "LiveLinkHubLog.h"
@@ -20,6 +21,7 @@
 #include "LiveLinkTypes.h"
 #include "Misc/FileHelper.h"
 #include "Misc/PackageName.h"
+#include "Misc/ScopeExit.h"
 #include "NamingTokensEditorSubsystem.h"
 #include "Recording/Implementations/LiveLinkUAssetRecording.h"
 #include "Recording/LiveLinkHubPlaybackController.h"
@@ -35,7 +37,7 @@
 #include "Widgets/Layout/SBox.h"
 #include "Widgets/SWidget.h"
 
-#define LOCTEXT_NAMESPACE "LiveLinkHub.RecordingController"
+#define LOCTEXT_NAMESPACE "LiveLinkHub.AssetRecorder"
 
 
 namespace UAssetRecorderUtils
@@ -92,6 +94,53 @@ bool FLiveLinkUAssetRecorder::IsRecording() const
 	return bIsRecording;
 }
 
+bool FLiveLinkUAssetRecorder::CanRecord(FText* OutErrorMessage) const
+{
+	FText ErrorText;
+	
+	ON_SCOPE_EXIT
+	{
+		if (OutErrorMessage)
+		{
+			*OutErrorMessage = MoveTemp(ErrorText);
+		}
+	};
+	
+	const FLiveLinkHubModule& LiveLinkHubModule = FModuleManager::Get().GetModuleChecked<FLiveLinkHubModule>("LiveLinkHub");
+	if (LiveLinkHubModule.GetPlaybackController()->IsInPlayback())
+	{
+		return false;
+	}
+	
+	const ILiveLinkRecordingSessionInfo& SessionInfo = ILiveLinkRecordingSessionInfo::Get();
+	if (SessionInfo.GetSessionName().IsEmpty())
+	{
+		ErrorText = LOCTEXT("SessionEmptyTooltip", "Session cannot be empty");
+		return false;
+	}
+
+	if (!FPaths::ValidatePath(SessionInfo.GetSessionName(), &ErrorText))
+	{
+		ErrorText = FText::Format(LOCTEXT("SessionErrorTooltip", "Session error: {0}"), ErrorText);
+		return false;
+	}
+
+	if (SessionInfo.GetSlateName().IsEmpty())
+	{
+		ErrorText = LOCTEXT("SlateEmptyTooltip", "Slate cannot be empty");
+		return false;
+	}
+
+	if (!FPaths::ValidatePath(SessionInfo.GetSlateName(), &ErrorText))
+	{
+		ErrorText = FText::Format(LOCTEXT("SlateErrorTooltip", "Slate error: {0}"), ErrorText);
+		return false;
+	}
+	
+	// TODO: Check for recording destination asset naming conflict.
+	return true;
+}
+
 void FLiveLinkUAssetRecorder::RecordBaseData(FLiveLinkRecordingBaseDataContainer& StaticDataContainer, TSharedPtr<FInstancedStruct>&& DataToRecord)
 {
 	const double TimeNowInSeconds = FPlatformTime::Seconds();
@@ -268,6 +317,10 @@ void FLiveLinkUAssetRecorder::SaveRecording()
 			MakeUnique<FAsyncTask<FLiveLinkSaveRecordingAsyncTask>>(NewRecording, this));
 		AsyncTask.StartBackgroundTask();
 	}
+
+	// Increment take on successful recording.
+	ILiveLinkRecordingSessionInfo& SessionInfo = ILiveLinkRecordingSessionInfo::Get();
+	SessionInfo.SetTakeNumber(SessionInfo.GetTakeNumber() + 1);
 }
 
 void FLiveLinkUAssetRecorder::RecordInitialStaticData()
diff --git a/Engine/Plugins/Animation/LiveLinkHub/Source/LiveLinkHub/Private/Recording/Implementations/LiveLinkUAssetRecorder.h b/Engine/Plugins/Animation/LiveLinkHub/Source/LiveLinkHub/Private/Recording/Implementations/LiveLinkUAssetRecorder.h
--- a/Engine/Plugins/Animation/LiveLinkHub/Source/LiveLinkHub/Private/Recording/Implementations/LiveLinkUAssetRecorder.h
+++ b/Engine/Plugins/Animation/LiveLinkHub/Source/LiveLinkHub/Private/Recording/Implementations/LiveLinkUAssetRecorder.h
@@ -24,6 +24,7 @@ public:
 	virtual void RecordStaticData(const FLiveLinkSubjectKey& SubjectKey, TSubclassOf<ULiveLinkRole> Role, const FLiveLinkStaticDataStruct& StaticData) override;
 	virtual void RecordFrameData(const FLiveLinkSubjectKey& SubjectKey, const FLiveLinkFrameDataStruct& FrameData) override;
 	virtual bool IsRecording() const override;
+	virtual bool CanRecord(FText* OutErrorMessage) const override;
 	//~ End ILiveLinkRecorder
 
 private:
diff --git a/Engine/Plugins/Animation/LiveLinkHub/Source/LiveLinkHub/Private/Recording/LiveLinkHubRecordingController.cpp b/Engine/Plugins/Animation/LiveLinkHub/Source/LiveLinkHub/Private/Recording/LiveLinkHubRecordingController.cpp
--- a/Engine/Plugins/Animation/LiveLinkHub/Source/LiveLinkHub/Private/Recording/LiveLinkHubRecordingController.cpp
+++ b/Engine/Plugins/Animation/LiveLinkHub/Source/LiveLinkHub/Private/Recording/LiveLinkHubRecordingController.cpp
@@ -10,6 +10,7 @@
 #include "Widgets/Layout/SBox.h"
 #include "Widgets/SWidget.h"
 
+#define LOCTEXT_NAMESPACE "LiveLinkHub.RecordingController"
 
 FLiveLinkHubRecordingController::FLiveLinkHubRecordingController()
 {
@@ -22,7 +23,8 @@ TSharedRef<SWidget> FLiveLinkHubRecordingController::MakeRecordToolbarEntry()
 		.CanRecord_Raw(this, &FLiveLinkHubRecordingController::CanRecord)
 		.IsRecording_Raw(this, &FLiveLinkHubRecordingController::IsRecording)
 		.OnStartRecording_Raw(this, &FLiveLinkHubRecordingController::StartRecording)
-		.OnStopRecording_Raw(this, &FLiveLinkHubRecordingController::StopRecording);
+		.OnStopRecording_Raw(this, &FLiveLinkHubRecordingController::StopRecording)
+		.ToolTipText_Raw(this, &FLiveLinkHubRecordingController::GetRecordingButtonTooltip);
 }
 
 void FLiveLinkHubRecordingController::StartRecording()
@@ -46,15 +48,19 @@ bool FLiveLinkHubRecordingController::IsRecording() const
 
 bool FLiveLinkHubRecordingController::CanRecord() const
 {
-	const FLiveLinkHubModule& LiveLinkHubModule = FModuleManager::Get().GetModuleChecked<FLiveLinkHubModule>("LiveLinkHub");
-	if (LiveLinkHubModule.GetPlaybackController()->IsInPlayback())
+	return RecorderImplementation->CanRecord(&ErrorMessage);
+}
+
+FText FLiveLinkHubRecordingController::GetRecordingButtonTooltip() const
+{
+	if (ErrorMessage.IsEmpty())
 	{
-		return false;
+		return IsRecording() ?
+			LOCTEXT("RecordingStopTooltip", "Stop a recording") :
+			LOCTEXT("RecordingStartTooltip", "Start a recording");
 	}
-
-	// TODO: Check for recording destination asset naming conflict.
-
-	return true;
+	
+	return ErrorMessage;
 }
 
 void FLiveLinkHubRecordingController::RecordStaticData(const FLiveLinkSubjectKey& SubjectKey, TSubclassOf<ULiveLinkRole> Role, const FLiveLinkStaticDataStruct& StaticData)
@@ -66,3 +72,5 @@ void FLiveLinkHubRecordingController::RecordFrameData(const FLiveLinkSubjectKey&
 {
 	RecorderImplementation->RecordFrameData(SubjectKey, FrameData);
 }
+
+#undef LOCTEXT_NAMESPACE
diff --git a/Engine/Plugins/Animation/LiveLinkHub/Source/LiveLinkHub/Private/Recording/LiveLinkHubRecordingController.h b/Engine/Plugins/Animation/LiveLinkHub/Source/LiveLinkHub/Private/Recording/LiveLinkHubRecordingController.h
--- a/Engine/Plugins/Animation/LiveLinkHub/Source/LiveLinkHub/Private/Recording/LiveLinkHubRecordingController.h
+++ b/Engine/Plugins/Animation/LiveLinkHub/Source/LiveLinkHub/Private/Recording/LiveLinkHubRecordingController.h
@@ -30,6 +30,9 @@ public:
 	/** Returns whether we can begin recording. */
 	bool CanRecord() const;
 
+	/** The tool tip to display on the recording button. */
+	FText GetRecordingButtonTooltip() const;
+
 	/** Record static data in the current recording. */
 	void RecordStaticData(const FLiveLinkSubjectKey& SubjectKey, TSubclassOf<class ULiveLinkRole> Role, const FLiveLinkStaticDataStruct& StaticData);
 
@@ -38,4 +41,7 @@ public:
 private:
 	/** Recorder used to serialize livelink data into a given format. */
 	TSharedPtr<ILiveLinkRecorder> RecorderImplementation;
+	
+	/** Any error text associated with the recording, for display on the recording button tooltip. */
+	mutable FText ErrorMessage;
 };
diff --git a/Engine/Plugins/Animation/LiveLinkHub/Source/LiveLinkHub/Private/Recording/LiveLinkRecorder.h b/Engine/Plugins/Animation/LiveLinkHub/Source/LiveLinkHub/Private/Recording/LiveLinkRecorder.h
--- a/Engine/Plugins/Animation/LiveLinkHub/Source/LiveLinkHub/Private/Recording/LiveLinkRecorder.h
+++ b/Engine/Plugins/Animation/LiveLinkHub/Source/LiveLinkHub/Private/Recording/LiveLinkRecorder.h
@@ -23,4 +23,6 @@ public:
 	virtual void RecordFrameData(const FLiveLinkSubjectKey& SubjectKey, const FLiveLinkFrameDataStruct& FrameData) = 0;
 	/** Return whether we are currently recording livelink data. */
 	virtual bool IsRecording() const = 0;
+	/** Returns whether we can begin recording. */
+	virtual bool CanRecord(FText* OutErrorMessage = nullptr) const = 0;
 };
