LiveLinkHub - Add playback controls to the bottom status bar

When a recording a is playing back, a small playback control will appear at the bottom of the app.

diff --git a/Engine/Plugins/Animation/LiveLinkHub/Source/LiveLinkHub/Private/LiveLinkHubCreatorAppMode.h b/Engine/Plugins/Animation/LiveLinkHub/Source/LiveLinkHub/Private/LiveLinkHubCreatorAppMode.h
--- a/Engine/Plugins/Animation/LiveLinkHub/Source/LiveLinkHub/Private/LiveLinkHubCreatorAppMode.h
+++ b/Engine/Plugins/Animation/LiveLinkHub/Source/LiveLinkHub/Private/LiveLinkHubCreatorAppMode.h
@@ -412,13 +412,19 @@ public:
 
 	virtual TArray<TSharedRef<SWidget>> GetStatusBarWidgets_Impl() override
 	{
-		UToolMenu* Menu = UToolMenus::Get()->ExtendMenu(TEXT("LiveLinkHub.ToolBar"));
-
-		FToolMenuSection& SessionConfigSection = Menu->AddSection(TEXT("SessionConfig"), FText::GetEmpty(), FToolMenuInsert(NAME_None, EToolMenuInsertType::First));
-
 		return {
-			SNew(SBox)
-			.Padding(0, 4.0, 0, 0)
+			SNew(SHorizontalBox)
+			+ SHorizontalBox::Slot()
+			.AutoWidth()
+			.VAlign(VAlign_Center)
+			.Padding(FMargin(2.0, 0.0, 4.0, 0.0))
+			[
+				FLiveLinkHub::Get()->GetPlaybackController()->MakePlaybackControlWidget()
+			]
+			+ SHorizontalBox::Slot()
+			.AutoWidth()
+			.VAlign(VAlign_Center)
+			.Padding(2.0, 0.0)
 			[
 				SNew(STextBlock)
 				.Text_Raw(this, &FLiveLinkHubCreatorAppMode::GetLoadedConfigText)
diff --git a/Engine/Plugins/Animation/LiveLinkHub/Source/LiveLinkHub/Private/Recording/LiveLinkHubPlaybackController.cpp b/Engine/Plugins/Animation/LiveLinkHub/Source/LiveLinkHub/Private/Recording/LiveLinkHubPlaybackController.cpp
--- a/Engine/Plugins/Animation/LiveLinkHub/Source/LiveLinkHub/Private/Recording/LiveLinkHubPlaybackController.cpp
+++ b/Engine/Plugins/Animation/LiveLinkHub/Source/LiveLinkHub/Private/Recording/LiveLinkHubPlaybackController.cpp
@@ -19,6 +19,7 @@
 #include "LiveLinkHubModule.h"
 #include "Modules/ModuleManager.h"
 #include "PackageTools.h"
+#include "Recording/SLiveLinkHubPlaybackControl.h"
 #include "UI/Widgets/SLiveLinkHubPlaybackWidget.h"
 #include "UObject/Object.h"
 #include "UObject/Package.h"
@@ -136,6 +137,17 @@ TSharedRef<SWidget> FLiveLinkHubPlaybackController::MakePlaybackWidget()
 		.GetFrameRate_Raw(this, &FLiveLinkHubPlaybackController::GetFrameRate);
 }
 
+TSharedRef<SWidget> FLiveLinkHubPlaybackController::MakePlaybackControlWidget()
+{
+	return SNew(SLiveLinkHubPlaybackControl)
+		.OnGetPlaybackState_Raw(this, &FLiveLinkHubPlaybackController::OnGetPlaybackState)
+		.OnResumePlayback_Raw(this, &FLiveLinkHubPlaybackController::OnResumePlayback)
+		.OnPausePlayback_Raw(this, &FLiveLinkHubPlaybackController::OnPausePlayback)
+		.OnExitPlayback_Raw(this, &FLiveLinkHubPlaybackController::OnExitPlayback)
+		.OnGetPlaybackTime_Raw(this, &FLiveLinkHubPlaybackController::OnGetPlaybackTime);
+}
+
+
 void FLiveLinkHubPlaybackController::StartPlayback()
 {
 	ResumePlayback();
@@ -650,3 +662,49 @@ bool FLiveLinkHubPlaybackController::ShouldRestart() const
 	return RecordingToPlay.IsValid() && ((bIsReverse && CurrentFrame <= GetSelectionStartTime().Time.GetFrame())
 		|| (!bIsReverse && CurrentFrame >= GetSelectionEndTime().Time.GetFrame()));
 }
+
+ELiveLinkHubPlaybackState FLiveLinkHubPlaybackController::OnGetPlaybackState() const
+{
+	ELiveLinkHubPlaybackState PlaybackState = ELiveLinkHubPlaybackState::NotInPlayback;
+
+	if (IsInPlayback())
+	{
+		if (IsPlaying() && !IsPaused())
+		{
+			PlaybackState = ELiveLinkHubPlaybackState::Playing;
+		}
+		else if (IsPaused())
+		{
+			PlaybackState = ELiveLinkHubPlaybackState::Paused;
+		}
+	}
+
+	return PlaybackState;
+}
+
+void FLiveLinkHubPlaybackController::OnResumePlayback()
+{
+	if (ensureAlways(IsInPlayback() && IsPaused()))
+	{
+		constexpr bool bInReverse = false;
+		BeginPlayback(bInReverse);
+	}
+}
+
+void FLiveLinkHubPlaybackController::OnPausePlayback()
+{
+	if (ensureAlways(IsInPlayback() && !IsPaused()))
+	{
+		PausePlayback();
+	}
+}
+
+void FLiveLinkHubPlaybackController::OnExitPlayback()
+{
+	EjectAndUnload();
+}
+
+FQualifiedFrameTime FLiveLinkHubPlaybackController::OnGetPlaybackTime() const
+{
+	return GetCurrentTime();
+}
\ No newline at end of file
diff --git a/Engine/Plugins/Animation/LiveLinkHub/Source/LiveLinkHub/Private/Recording/LiveLinkHubPlaybackController.h b/Engine/Plugins/Animation/LiveLinkHub/Source/LiveLinkHub/Private/Recording/LiveLinkHubPlaybackController.h
--- a/Engine/Plugins/Animation/LiveLinkHub/Source/LiveLinkHub/Private/Recording/LiveLinkHubPlaybackController.h
+++ b/Engine/Plugins/Animation/LiveLinkHub/Source/LiveLinkHub/Private/Recording/LiveLinkHubPlaybackController.h
@@ -16,6 +16,7 @@
 
 #include "LiveLinkRecordingRangeHelpers.h"
 
+enum class ELiveLinkHubPlaybackState : uint8;
 class ILiveLinkClient;
 struct FLiveLinkRecordingData;
 class FRunnableThread;
@@ -34,6 +35,9 @@ public:
 	/** Create the playback widget. */
 	TSharedRef<SWidget> MakePlaybackWidget();
 
+	/** Create a lightweight playback widget. */
+	TSharedRef<SWidget> MakePlaybackControlWidget();
+
 	/** Apply the recording's preset then prepare the data needed to playback. */
 	void PreparePlayback(ULiveLinkRecording* InLiveLinkRecording);
 	
@@ -197,6 +201,21 @@ private:
 	/** Checks if the current playback settings indicates the recording should restart. */
 	bool ShouldRestart() const;
 
+	/** Get the state of the playback. */
+	ELiveLinkHubPlaybackState OnGetPlaybackState() const;
+
+	/** Handles the user clicking the play button to resume playback. */
+	void OnResumePlayback();
+
+	/** Handles the user clicking the pause button to resume playback. */
+	void OnPausePlayback();
+
+	/** Handles the user clicking the eject button to exit playback. */
+	void OnExitPlayback();
+
+	/** Returns the current time in the playback. */
+	FQualifiedFrameTime OnGetPlaybackTime() const;
+
 private:
 	/** If the system has established a connection with the client. */
 	bool bIsReady = false;
diff --git a/Engine/Plugins/Animation/LiveLinkHub/Source/LiveLinkHub/Private/Recording/LiveLinkHubRecordingListController.h b/Engine/Plugins/Animation/LiveLinkHub/Source/LiveLinkHub/Private/Recording/LiveLinkHubRecordingListController.h
--- a/Engine/Plugins/Animation/LiveLinkHub/Source/LiveLinkHub/Private/Recording/LiveLinkHubRecordingListController.h
+++ b/Engine/Plugins/Animation/LiveLinkHub/Source/LiveLinkHub/Private/Recording/LiveLinkHubRecordingListController.h
@@ -28,9 +28,7 @@ public:
 	TSharedRef<SWidget> MakeRecordingList()
 	{
 		return SNew(SLiveLinkHubRecordingListView)
-			.OnImportRecording_Raw(this, &FLiveLinkHubRecordingListController::OnImportRecording)
-			.OnEject_Raw(this, &FLiveLinkHubRecordingListController::OnEjectRecording)
-			.CanEject_Raw(this, &FLiveLinkHubRecordingListController::CanEjectRecording);
+			.OnImportRecording_Raw(this, &FLiveLinkHubRecordingListController::OnImportRecording);
 	}
 
 private:
@@ -65,20 +63,6 @@ private:
 		}
 	}
 
-	void OnEjectRecording()
-	{
-		if (const TSharedPtr<FLiveLinkHub> HubPtr = LiveLinkHub.Pin())
-		{
-			HubPtr->GetPlaybackController()->EjectAndUnload();
-		}
-	}
-
-	bool CanEjectRecording() const
-	{
-		const TSharedPtr<FLiveLinkHub> HubPtr = LiveLinkHub.Pin();
-		return HubPtr && HubPtr->GetPlaybackController()->GetRecording().IsValid();
-	}
-
 private:
 	/** LiveLinkHub object that holds the different controllers. */
 	TWeakPtr<FLiveLinkHub> LiveLinkHub;
diff --git a/Engine/Plugins/Animation/LiveLinkHub/Source/LiveLinkHub/Private/Recording/SLiveLinkHubPlaybackControl.h b/Engine/Plugins/Animation/LiveLinkHub/Source/LiveLinkHub/Private/Recording/SLiveLinkHubPlaybackControl.h
new file mode 100644
--- /dev/null
+++ b/Engine/Plugins/Animation/LiveLinkHub/Source/LiveLinkHub/Private/Recording/SLiveLinkHubPlaybackControl.h
@@ -0,0 +1,179 @@
+// Copyright Epic Games, Inc. All Rights Reserved.
+
+#pragma once
+
+#include "Widgets/SCompoundWidget.h"
+
+#include "LiveLinkHubPlaybackController.h"
+#include "Widgets/Images/SImage.h"
+#include "Widgets/Input/SButton.h"
+#include "Widgets/SBoxPanel.h"
+#include "Widgets/Text/STextBlock.h"
+
+class FLiveLinkHub;
+
+#define LOCTEXT_NAMESPACE "LiveLinkHub.PlaybackControl"
+
+
+DECLARE_DELEGATE_RetVal(FQualifiedFrameTime, FOnGetPlaybackTime);
+
+enum class ELiveLinkHubPlaybackState : uint8
+{
+	NotInPlayback,
+	Playing,
+	Paused
+};
+
+DECLARE_DELEGATE_RetVal(ELiveLinkHubPlaybackState, FOnGetPlaybackState);
+
+/**
+ * Lightweight playback control widget.
+ */
+class SLiveLinkHubPlaybackControl : public SCompoundWidget
+{
+public:
+	DECLARE_DELEGATE_RetVal(bool, FSimplePredicate);
+
+	SLATE_BEGIN_ARGS(SLiveLinkHubPlaybackControl) {}
+		SLATE_EVENT(FOnGetPlaybackState, OnGetPlaybackState)
+		SLATE_EVENT(FSimpleDelegate, OnResumePlayback)
+		SLATE_EVENT(FSimpleDelegate, OnPausePlayback)
+		SLATE_EVENT(FSimpleDelegate, OnExitPlayback)
+		SLATE_EVENT(FOnGetPlaybackTime, OnGetPlaybackTime)
+	SLATE_END_ARGS()
+
+	/**
+		* @param InArgs
+		* @param InStatusBarId Unique ID needed for the status bar
+		*/
+	void Construct(const FArguments& InArgs)
+	{
+		OnResumePlaybackDelegate = InArgs._OnResumePlayback;
+		OnPausePlaybackDelegate = InArgs._OnPausePlayback;
+		OnExitPlaybackDelegate = InArgs._OnExitPlayback;
+		OnGetPlaybackTimeDelegate = InArgs._OnGetPlaybackTime;
+		OnGetPlaybackStateDelegate = InArgs._OnGetPlaybackState;
+
+		ChildSlot
+		[
+			SNew(SHorizontalBox)
+			+ SHorizontalBox::Slot()
+			.VAlign(VAlign_Center)
+			.Padding(2.0f)
+			.AutoWidth()
+			[
+				SNew(STextBlock)
+				.Visibility(this, &SLiveLinkHubPlaybackControl::GetPlaybackControlsVisibility)
+				.Text(this, &SLiveLinkHubPlaybackControl::GetPlaybackTime)
+			]
+			+ SHorizontalBox::Slot()
+			.AutoWidth()
+			[
+				SNew(SButton)
+				.OnClicked(this, &SLiveLinkHubPlaybackControl::OnClickPlaybackButton)
+				.Visibility(this, &SLiveLinkHubPlaybackControl::GetPlaybackControlsVisibility)
+				.ToolTipText_Lambda([this]()
+				{
+					return IsPaused() ? LOCTEXT("PlayButton", "Play") : LOCTEXT("PauseButton", "Pause");
+				})
+				.ButtonStyle(FAppStyle::Get(), "Animation.PlayControlsButton")
+				.ContentPadding(0.0f)
+				[
+					SNew(SImage)
+						.ColorAndOpacity(FSlateColor::UseSubduedForeground())
+						.Image(this, &SLiveLinkHubPlaybackControl::GetPlayIcon)
+				]
+			]
+			+ SHorizontalBox::Slot()
+			.AutoWidth()
+			[
+				SNew(SButton)
+				.ButtonStyle(FAppStyle::Get(), "Animation.PlayControlsButton")
+				.ContentPadding(0.0f)
+				.Visibility(this, &SLiveLinkHubPlaybackControl::GetPlaybackControlsVisibility)
+				.OnClicked(this, &SLiveLinkHubPlaybackControl::OnClickExitPlaybackButton)
+				.ToolTipText(LOCTEXT("EjectButtonToolTip", "Exit Playback"))
+				[
+					SNew(SImage)
+						.Image(FAppStyle::Get().GetBrush("LevelViewport.EjectActorPilot.Small"))
+						.ColorAndOpacity(FSlateColor::UseSubduedForeground())
+				]
+			]
+		];
+	}
+
+	/** Get a formatted text version of the current playback time. */
+	FText GetPlaybackTime() const
+	{
+		FText FormattedText = INVTEXT("0:00");
+		if (OnGetPlaybackTimeDelegate.IsBound())
+		{
+			const FQualifiedFrameTime Time = OnGetPlaybackTimeDelegate.Execute();
+			const FTimecode Timecode = Time.ToTimecode();
+
+			TStringBuilder<64> Builder;
+
+			if (Timecode.Hours != 0)
+			{
+				Builder.Appendf(TEXT("%02d:%02d:%02d;%02d"), FMath::Abs(Timecode.Hours), FMath::Abs(Timecode.Minutes), FMath::Abs(Timecode.Seconds), FMath::Abs(Timecode.Frames));
+			}
+			else
+			{
+				Builder.Appendf(TEXT("%02d:%02d;%02d"), FMath::Abs(Timecode.Minutes), FMath::Abs(Timecode.Seconds), FMath::Abs(Timecode.Frames));
+			}
+
+			FormattedText = FText::FromString(Builder.ToString());
+		}
+
+		return FormattedText;
+	}
+
+	/** Handles user clicking either the play/pause button. */
+	FReply OnClickPlaybackButton()
+	{
+		if (IsPaused())
+		{
+			OnResumePlaybackDelegate.ExecuteIfBound();
+		}
+		else
+		{
+			OnPausePlaybackDelegate.ExecuteIfBound();
+		}
+		return FReply::Handled();
+	}
+
+	/** Handles the user clicking the eject button to exit playback. */
+	FReply OnClickExitPlaybackButton()
+	{
+		OnExitPlaybackDelegate.ExecuteIfBound();
+		return FReply::Handled();
+	}
+
+	/** Returns whether the current recording playing back is paused. */
+	bool IsPaused() const
+	{
+		return OnGetPlaybackStateDelegate.Execute() == ELiveLinkHubPlaybackState::Paused;
+	}
+
+	/** Get the icon for the play/pause button. */
+	const FSlateBrush* GetPlayIcon() const
+	{
+		return IsPaused() ? FAppStyle::Get().GetBrush("Animation.Forward") : FAppStyle::Get().GetBrush("Animation.Pause");
+	}
+
+	/** Get the visibility of the controls (Only visible when a recording is playing back.) */
+	EVisibility GetPlaybackControlsVisibility() const
+	{
+		return OnGetPlaybackStateDelegate.Execute() != ELiveLinkHubPlaybackState::NotInPlayback ? EVisibility::Visible : EVisibility::Collapsed;
+	}
+
+private:
+	//~ Playback Delegates 
+	FOnGetPlaybackState OnGetPlaybackStateDelegate;
+	FSimpleDelegate OnResumePlaybackDelegate;
+	FSimpleDelegate OnPausePlaybackDelegate;
+	FSimpleDelegate OnExitPlaybackDelegate;
+	FOnGetPlaybackTime OnGetPlaybackTimeDelegate;
+};
+
+#undef LOCTEXT_NAMESPACE /* LiveLinkHub.PlaybackControl */
diff --git a/Engine/Plugins/Animation/LiveLinkHub/Source/LiveLinkHub/Private/Recording/SLiveLinkHubRecordingListView.h b/Engine/Plugins/Animation/LiveLinkHub/Source/LiveLinkHub/Private/Recording/SLiveLinkHubRecordingListView.h
--- a/Engine/Plugins/Animation/LiveLinkHub/Source/LiveLinkHub/Private/Recording/SLiveLinkHubRecordingListView.h
+++ b/Engine/Plugins/Animation/LiveLinkHub/Source/LiveLinkHub/Private/Recording/SLiveLinkHubRecordingListView.h
@@ -29,14 +29,10 @@ class SLiveLinkHubRecordingListView : public SCompoundWidget
 {
 public:
 	DECLARE_DELEGATE_OneParam(FOnImportRecording, const struct FAssetData&);
-	DECLARE_DELEGATE(FOnEject);
-	DECLARE_DELEGATE_RetVal(bool, FCanEject)
 	
 	SLATE_BEGIN_ARGS(SLiveLinkHubRecordingListView)
 		{}
 		SLATE_EVENT(FOnImportRecording, OnImportRecording)
-		SLATE_EVENT(FOnEject, OnEject)
-		SLATE_EVENT(FCanEject, CanEject)
 	SLATE_END_ARGS()
 
 	SLiveLinkHubRecordingListView()
@@ -66,63 +62,31 @@ public:
 	void Construct(const FArguments& InArgs)
 	{
 		OnImportRecordingDelegate = InArgs._OnImportRecording;
-		OnEjectDelegate = InArgs._OnEject;
-		OnCanEjectDelegate = InArgs._CanEject;
 		
 		ChildSlot
 		[
 			SNew(SVerticalBox)
 			+ SVerticalBox::Slot()
-			.AutoHeight()
+			.VAlign(VAlign_Fill)
 			[
-				SNew(SHorizontalBox)
-				+ SHorizontalBox::Slot()
-				.HAlign(HAlign_Fill)
-				.VAlign(VAlign_Center)
-				.Padding(4.f)
+				SNew(SBox)
+				.Visibility(this, &SLiveLinkHubRecordingListView::GetRecordingPickerVisibility)
 				[
-					SNew(STextBlock)
-					.Text(LOCTEXT("RecordingTitle", "Recordings"))
-				]
-				+ SHorizontalBox::Slot()
-				.AutoWidth()
-				.HAlign(HAlign_Right)
-				.VAlign(VAlign_Center)
-				.Padding(4.f)
-				[
-					SNew(SButton)
-					.Visibility(this, &SLiveLinkHubRecordingListView::GetRecordingPickerVisibility)
-					.Text(LOCTEXT("EjectButton", "Exit Playback"))
-					.OnClicked(this, &SLiveLinkHubRecordingListView::OnEjectClicked)
-					.IsEnabled(this, &SLiveLinkHubRecordingListView::CanEjectRecording)
+					CreateRecordingPicker()
 				]
 			]
-			+ SVerticalBox::Slot()
-			.FillHeight(1.0f)
+			+SVerticalBox::Slot()
+			.AutoHeight()
+			.VAlign(VAlign_Top)
+			.HAlign(HAlign_Center)
 			[
-				SNew(SVerticalBox)
-				+ SVerticalBox::Slot()
-				.VAlign(VAlign_Fill)
-				[
-					SNew(SBox)
-					.Visibility(this, &SLiveLinkHubRecordingListView::GetRecordingPickerVisibility)
-					[
-						CreateRecordingPicker()
-					]
-				]
-				+SVerticalBox::Slot()
-				.AutoHeight()
-				.VAlign(VAlign_Top)
-				.HAlign(HAlign_Center)
-				[
-					SNew(STextBlock)
-					.Visibility_Lambda([this]()
-					{
-						EVisibility RecordingPickerVisibility = GetRecordingPickerVisibility();
-						return RecordingPickerVisibility == EVisibility::Visible ? EVisibility::Collapsed : EVisibility::Visible;
-					})
-					.Text(GetNoAssetsWarningText())
-				]
+				SNew(STextBlock)
+				.Visibility_Lambda([this]()
+				{
+					EVisibility RecordingPickerVisibility = GetRecordingPickerVisibility();
+					return RecordingPickerVisibility == EVisibility::Visible ? EVisibility::Collapsed : EVisibility::Visible;
+				})
+				.Text(GetNoAssetsWarningText())
 			]
 		];
 	}
@@ -135,18 +99,6 @@ private:
 		OnImportRecordingDelegate.Execute(AssetData);
 	}
 
-	FReply OnEjectClicked()
-	{
-		OnEjectDelegate.ExecuteIfBound();
-		return FReply::Handled();
-	}
-
-	bool CanEjectRecording() const
-	{
-		check(OnCanEjectDelegate.IsBound());
-		return OnCanEjectDelegate.Execute();
-	}
-
 	/** When an asset is added to the asset registry. */
 	void OnAssetAdded(const FAssetData& InAssetData)
 	{
@@ -352,10 +304,6 @@ private:
 private:
 	/** Delegate used for noticing the hub that a recording was selected for playback. */
 	FOnImportRecording OnImportRecordingDelegate;
-	/** Delegate used when ejecting a recording. */
-	FOnEject OnEjectDelegate;
-	/** Delegate used to determine a recording can be ejected. */
-	FCanEject OnCanEjectDelegate;
 	/** The asset picker used for selecting recordings. */
 	TSharedPtr<SWidget> AssetPicker;
 	
