LiveLinkHub - Fix crash when launching the app

#rnx

diff --git a/Engine/Plugins/Animation/LiveLinkHub/Source/LiveLinkHub/Private/LiveLinkHub.cpp b/Engine/Plugins/Animation/LiveLinkHub/Source/LiveLinkHub/Private/LiveLinkHub.cpp
--- a/Engine/Plugins/Animation/LiveLinkHub/Source/LiveLinkHub/Private/LiveLinkHub.cpp
+++ b/Engine/Plugins/Animation/LiveLinkHub/Source/LiveLinkHub/Private/LiveLinkHub.cpp
@@ -7,6 +7,7 @@
 #include "AssetRegistry/IAssetRegistry.h"
 #include "Clients/LiveLinkHubClientsController.h"
 #include "Clients/LiveLinkHubProvider.h"
+#include "EditorViewportCommands.h"
 #include "Features/IModularFeatures.h"
 #include "Framework/Application/SlateApplication.h"
 #include "HAL/PlatformProcess.h"
@@ -114,6 +115,8 @@ void FLiveLinkHub::Initialize(FLiveLinkHubTicker& Ticker)
 	FModuleManager::Get().LoadModule("Settings");
 	FModuleManager::Get().LoadModule("StatusBar");
 
+	FEditorViewportCommands::Register();
+
 	if (LiveLinkHubInternalUtils::IsDistributedBuild())
 	{
 		IAssetRegistry& AssetRegistry = FModuleManager::LoadModuleChecked<FAssetRegistryModule>(AssetRegistryConstants::ModuleName).Get();
