Fix crash on startup caused by missing context for registering asset editor viewport commands

#rnx

diff --git a/Engine/Plugins/Animation/LiveLinkHub/Source/LiveLinkHub/Private/LiveLinkHub.cpp b/Engine/Plugins/Animation/LiveLinkHub/Source/LiveLinkHub/Private/LiveLinkHub.cpp
--- a/Engine/Plugins/Animation/LiveLinkHub/Source/LiveLinkHub/Private/LiveLinkHub.cpp
+++ b/Engine/Plugins/Animation/LiveLinkHub/Source/LiveLinkHub/Private/LiveLinkHub.cpp
@@ -7,7 +7,6 @@
 #include "AssetRegistry/IAssetRegistry.h"
 #include "Clients/LiveLinkHubClientsController.h"
 #include "Clients/LiveLinkHubProvider.h"
-#include "EditorViewportCommands.h"
 #include "Features/IModularFeatures.h"
 #include "Framework/Application/SlateApplication.h"
 #include "HAL/PlatformProcess.h"
@@ -115,8 +114,6 @@ void FLiveLinkHub::Initialize(FLiveLinkHubTicker& Ticker)
 	FModuleManager::Get().LoadModule("Settings");
 	FModuleManager::Get().LoadModule("StatusBar");
 
-	FEditorViewportCommands::Register();
-
 	if (LiveLinkHubInternalUtils::IsDistributedBuild())
 	{
 		IAssetRegistry& AssetRegistry = FModuleManager::LoadModuleChecked<FAssetRegistryModule>(AssetRegistryConstants::ModuleName).Get();
diff --git a/Engine/Plugins/Animation/LiveLinkHub/Source/LiveLinkHub/Private/LiveLinkHubModule.cpp b/Engine/Plugins/Animation/LiveLinkHub/Source/LiveLinkHub/Private/LiveLinkHubModule.cpp
--- a/Engine/Plugins/Animation/LiveLinkHub/Source/LiveLinkHub/Private/LiveLinkHubModule.cpp
+++ b/Engine/Plugins/Animation/LiveLinkHub/Source/LiveLinkHub/Private/LiveLinkHubModule.cpp
@@ -4,6 +4,7 @@
 
 #include "CoreGlobals.h"
 #include "Clients/LiveLinkHubProvider.h"
+#include "LevelEditor.h"
 #include "LiveLinkHubLog.h"
 #include "LiveLinkHubSubjectSettings.h"
 #include "LiveLinkSettings.h"
@@ -33,6 +34,9 @@ void FLiveLinkHubModule::StartupModule()
 
 	if (bCreateLiveLinkHubInstance)
 	{
+		// Needed for downstream modules that need the editor commands to be loaded.
+		FModuleManager::LoadModuleChecked<FLevelEditorModule>("LevelEditor");
+
 		if (!GetDefault<ULiveLinkHubSettings>()->bTickOnGameThread)
 		{
 			Ticker.StartTick();
