Add framerate display to live link hub (can be enabled in the settings)

diff --git a/Engine/Plugins/Animation/LiveLinkHub/Source/LiveLinkHub/Private/Settings/LiveLinkHubSettings.h b/Engine/Plugins/Animation/LiveLinkHub/Source/LiveLinkHub/Private/Settings/LiveLinkHubSettings.h
--- a/Engine/Plugins/Animation/LiveLinkHub/Source/LiveLinkHub/Private/Settings/LiveLinkHubSettings.h
+++ b/Engine/Plugins/Animation/LiveLinkHub/Source/LiveLinkHub/Private/Settings/LiveLinkHubSettings.h
@@ -59,6 +59,10 @@ public:
 	UPROPERTY(config, EditAnywhere, Category = "LiveLinkHub")
 	bool bRemoveInvalidSubjectsAfterLoadingSession = false;
 
+	/** Whether to show the app's frame rate in the top right corner. */
+	UPROPERTY(config, EditAnywhere, Category = "LiveLinkHub")
+	bool bShowFrameRate = false;
+
 	/** Whether to show memory usage in the top right corner. */
 	UPROPERTY(config, EditAnywhere, Category = "LiveLinkHub")
 	bool bShowMemoryUsage = true;
diff --git a/Engine/Plugins/Animation/LiveLinkHub/Source/LiveLinkHub/Private/UI/Widgets/SLiveLinkHubMemoryStats.h b/Engine/Plugins/Animation/LiveLinkHub/Source/LiveLinkHub/Private/UI/Widgets/SLiveLinkHubMemoryStats.h
--- a/Engine/Plugins/Animation/LiveLinkHub/Source/LiveLinkHub/Private/UI/Widgets/SLiveLinkHubMemoryStats.h
+++ b/Engine/Plugins/Animation/LiveLinkHub/Source/LiveLinkHub/Private/UI/Widgets/SLiveLinkHubMemoryStats.h
@@ -4,6 +4,7 @@
 
 #include "Widgets/SCompoundWidget.h"
 
+#include "Framework/Application/SlateApplication.h"
 #include "Framework/Notifications/NotificationManager.h"
 #include "LiveLinkHubLog.h"
 #include "Misc/Paths.h"
@@ -29,36 +30,68 @@ public:
 		ChildSlot
 		[
 			SNew(SHorizontalBox)
-			.Visibility_Lambda([]() { return GetDefault<ULiveLinkHubSettings>()->bShowMemoryUsage ? EVisibility::Visible : EVisibility::Collapsed; })
 			+ SHorizontalBox::Slot()
 			.VAlign(VAlign_Center)
+			.Padding(FMargin(0.0, 0.0, 5.0, 0.0))
 			.AutoWidth()
 			[
-				SNew(STextBlock)
-				.Text(LOCTEXT("MemoryLabel", "Mem: "))
+				SNew(SHorizontalBox)
+				.Visibility_Lambda([]() { return GetDefault<ULiveLinkHubSettings>()->bShowFrameRate ? EVisibility::Visible : EVisibility::Collapsed; })
+				+ SHorizontalBox::Slot()
+				.VAlign(VAlign_Bottom)
+				.AutoWidth()
+				[
+					SNew(STextBlock)
+						.Font(FAppStyle::GetFontStyle(TEXT("MainFrame.DebugTools.LabelFont")))
+						.ColorAndOpacity(FLinearColor(0.3f, 0.3f, 0.3f))
+						.Text(LOCTEXT("FrameRateLabel", "FPS: "))
+				]
+				+ SHorizontalBox::Slot()
+				.VAlign(VAlign_Bottom)
+				.AutoWidth()
+				[
+					SNew(STextBlock)
+						.Text(TAttribute<FText>::CreateLambda([this]() { return FText::Format(INVTEXT("{0} / {1}"), CachedFrameRateText, CachedFrameTimeText); }))
+				]
 			]
 			+ SHorizontalBox::Slot()
 			.VAlign(VAlign_Center)
 			.AutoWidth()
 			[
-				SNew(STextBlock)
-				.Text(TAttribute<FText>::CreateLambda([this]() { return CachedMemoryText; }))
-			] 
-			+ SHorizontalBox::Slot()
-			.AutoWidth()
-			.Padding(2.0f, 0.0f)
-			.VAlign(VAlign_Center)
-			[
-				SNew(SButton)
-				.ContentPadding(FMargin(0, 2.0))
-				.ToolTipText(LOCTEXT("MemoryWarningTooltip", "Live Link Hub is using an unusual amount of RAM. Click here to collect a memory trace."))
-				.Visibility(this, &SLiveLinkHubMemoryStats::OnGetWarningVisibility)
-				.OnClicked(this, &SLiveLinkHubMemoryStats::OnClickMemoryWarning)
-				.IsEnabled_Lambda([this]() { return !bCreatingTrace; })
+				SNew(SHorizontalBox)
+								.Visibility_Lambda([]() { return GetDefault<ULiveLinkHubSettings>()->bShowMemoryUsage  ? EVisibility::Visible : EVisibility::Collapsed; })
+				+ SHorizontalBox::Slot()
+				.VAlign(VAlign_Bottom)
+				.AutoWidth()
+				[
+					SNew(STextBlock)
+					.Font(FAppStyle::GetFontStyle(TEXT("MainFrame.DebugTools.LabelFont")))
+					.ColorAndOpacity(FLinearColor(0.3f, 0.3f, 0.3f))
+					.Text(LOCTEXT("MemoryLabel", "Mem: "))
+				]
+				+ SHorizontalBox::Slot()
+				.VAlign(VAlign_Bottom)
+				.AutoWidth()
 				[
-					SNew(SImage)
-					.ColorAndOpacity(FLinearColor::Red)
-					.Image(FCoreStyle::Get().GetBrush("Icons.Warning"))
+					SNew(STextBlock)
+						.Text(TAttribute<FText>::CreateLambda([this]() { return CachedMemoryText; }))
+				]
+				+ SHorizontalBox::Slot()
+				.AutoWidth()
+				.Padding(2.0f, 0.0f)
+				.VAlign(VAlign_Bottom)
+				[
+					SNew(SButton)
+						.ContentPadding(FMargin(0, 2.0))
+						.ToolTipText(LOCTEXT("MemoryWarningTooltip", "Live Link Hub is using an unusual amount of RAM. Click here to collect a memory trace."))
+						.Visibility(this, &SLiveLinkHubMemoryStats::OnGetWarningVisibility)
+						.OnClicked(this, &SLiveLinkHubMemoryStats::OnClickMemoryWarning)
+						.IsEnabled_Lambda([this]() { return !bCreatingTrace; })
+						[
+							SNew(SImage)
+								.ColorAndOpacity(FLinearColor::Red)
+								.Image(FCoreStyle::Get().GetBrush("Icons.Warning"))
+						]
 				]
 			]
 		];
@@ -67,6 +100,8 @@ public:
 	virtual void Tick(const FGeometry & AllottedGeometry, const double InCurrentTime, const float InDeltaTime) override
 	{
 		UpdateMemoryStats();
+		UpdateFrameTime();
+		UpdateFrameRate();
 
 		if (CachedTotalPhysicalUsedMB < GetDefault<ULiveLinkHubSettings>()->ShowMemoryWarningThresholdMB)
 		{
@@ -79,6 +114,36 @@ public:
 	}
 
 private:
+
+	/** Update the cached frame rate text */
+	void UpdateFrameRate() 
+	{
+		// Clamp to avoid huge averages at startup or after hitches
+		const float AverageFPS = 1.0f / FSlateApplication::Get().GetAverageDeltaTime();
+		const float ClampedFPS = (AverageFPS < 0.0f || AverageFPS > 4000.0f ) ? 0.0f : AverageFPS;
+
+		static const FNumberFormattingOptions FormatOptions = FNumberFormattingOptions()
+			.SetMinimumFractionalDigits(1)
+			.SetMaximumFractionalDigits(1);
+
+		CachedFrameRateText = FText::AsNumber(ClampedFPS, &FormatOptions);
+	}
+
+	/** Update the cached frame time text. */
+	void UpdateFrameTime()
+	{
+		// Clamp to avoid huge averages at startup or after hitches
+		const float AverageMS = FSlateApplication::Get().GetAverageDeltaTime() * 1000.0f;
+		const float ClampedMS = ( AverageMS < 0.0f || AverageMS > 4000.0f ) ? 0.0f : AverageMS;
+
+		static const FNumberFormattingOptions FormatOptions = FNumberFormattingOptions()
+			.SetMinimumFractionalDigits(1)
+			.SetMaximumFractionalDigits(1);
+		static const FText FrameTimeFmt = FText::FromString(TEXT("{0} ms"));
+
+		CachedFrameTimeText = FText::Format(FrameTimeFmt, FText::AsNumber( ClampedMS, &FormatOptions ));
+	}
+
 	/** Update our cached memory stats. */
 	void UpdateMemoryStats()
 	{
@@ -193,8 +258,13 @@ private:
 		return WarningMessage;
 	}
 
+private:
 	/** Cached memory used in this process. */
 	float CachedTotalPhysicalUsedMB = 0.f;
+	/** Cached text representation of the current app frame rate. */
+	FText CachedFrameRateText;
+	/** Cached text representation of the current app frame time. */
+	FText CachedFrameTimeText;
 	/** Cached text representation of our total physical memory used. */
 	FText CachedMemoryText;
 	/** Stores the path to the last trace snapshot that was generated. */
