Live Link Hub - Add menu entry to open logs folder
Clean up unused code in SLiveLinkHubTabViewWithManagerBase

diff --git a/Engine/Plugins/Animation/LiveLinkHub/Source/LiveLinkHub/Private/LiveLinkHub.cpp b/Engine/Plugins/Animation/LiveLinkHub/Source/LiveLinkHub/Private/LiveLinkHub.cpp
--- a/Engine/Plugins/Animation/LiveLinkHub/Source/LiveLinkHub/Private/LiveLinkHub.cpp
+++ b/Engine/Plugins/Animation/LiveLinkHub/Source/LiveLinkHub/Private/LiveLinkHub.cpp
@@ -9,6 +9,7 @@
 #include "Clients/LiveLinkHubProvider.h"
 #include "Features/IModularFeatures.h"
 #include "Framework/Application/SlateApplication.h"
+#include "GenericPlatform/GenericPlatformOutputDevices.h"
 #include "HAL/PlatformProcess.h"
 #include "ISettingsModule.h"
 #include "ILiveLinkHubMessagingModule.h"
@@ -341,6 +342,7 @@ void FLiveLinkHub::MapToolkitCommands()
 	CommandList->MapAction(Commands.SaveConfigAs, FExecuteAction::CreateSP(this, &FLiveLinkHub::SaveConfigAs));
 	CommandList->MapAction(Commands.SaveConfig, FExecuteAction::CreateSP(this, &FLiveLinkHub::SaveConfig),
 		FCanExecuteAction::CreateSP(this, &FLiveLinkHub::CanSaveConfig));
+	CommandList->MapAction(Commands.OpenLogsFolder, FExecuteAction::CreateSP(this, &FLiveLinkHub::OpenLogsFolder));
 
 	FLiveLinkHubApplication::MapToolkitCommands();
 }
@@ -410,6 +412,13 @@ void FLiveLinkHub::OpenConfig()
 	SessionManager->RestoreSession();
 }
 
+void FLiveLinkHub::OpenLogsFolder()
+{
+	const FString AbsoluteFilePath = FPaths::ConvertRelativePathToFull(FGenericPlatformOutputDevices::GetAbsoluteLogFilename());
+	const FString FileDirectory = FPaths::GetPath(AbsoluteFilePath);
+	FPlatformProcess::ExploreFolder(*FileDirectory);
+}
+
 void FLiveLinkHub::OnModularFeatureRegistered(const FName& Type, IModularFeature* ModularFeature)
 {
 	if (Type == ILiveLinkHubApplicationModeFactory::ModularFeatureName)
diff --git a/Engine/Plugins/Animation/LiveLinkHub/Source/LiveLinkHub/Private/LiveLinkHub.h b/Engine/Plugins/Animation/LiveLinkHub/Source/LiveLinkHub/Private/LiveLinkHub.h
--- a/Engine/Plugins/Animation/LiveLinkHub/Source/LiveLinkHub/Private/LiveLinkHub.h
+++ b/Engine/Plugins/Animation/LiveLinkHub/Source/LiveLinkHub/Private/LiveLinkHub.h
@@ -131,6 +131,9 @@ private:
 	/** Open an existing config. */
 	void OpenConfig();
 
+	/** Open the folder containing logs for the app. */
+	void OpenLogsFolder();
+
 	/** Handles adding application modes that are registered after LLH has finished initialization. */
 	void OnModularFeatureRegistered(const FName& Type, class IModularFeature* ModularFeature);
 	
diff --git a/Engine/Plugins/Animation/LiveLinkHub/Source/LiveLinkHub/Private/LiveLinkHubApplication.cpp b/Engine/Plugins/Animation/LiveLinkHub/Source/LiveLinkHub/Private/LiveLinkHubApplication.cpp
--- a/Engine/Plugins/Animation/LiveLinkHub/Source/LiveLinkHub/Private/LiveLinkHubApplication.cpp
+++ b/Engine/Plugins/Animation/LiveLinkHub/Source/LiveLinkHub/Private/LiveLinkHubApplication.cpp
@@ -225,6 +225,17 @@ void FLiveLinkHubApplication::RegisterMainMenu()
 		FText::GetEmpty()
 	);
 
+	UToolMenu* HelpEntry = MenuBar->AddSubMenu(
+		"MainMenu",
+		NAME_None,
+		"Help",
+		LOCTEXT("HelpMenu", "Help"),
+		FText::GetEmpty()
+	);
+
+	FToolMenuEntry OpenLogsFolderEntry = FToolMenuEntry::InitMenuEntry(FLiveLinkHubCommands::Get().OpenLogsFolder, TAttribute<FText>(), TAttribute<FText>(), FSlateIcon(FAppStyle::GetAppStyleSetName(), "Icons.FolderOpen"));
+	HelpEntry->AddMenuEntry(NAME_None, OpenLogsFolderEntry);
+
 	WindowEntry->AddDynamicSection(NAME_None, FNewSectionConstructChoice{FNewToolMenuDelegateLegacy::CreateRaw(this, &FLiveLinkHubApplication::FillWindowMenu)});
 }
 
diff --git a/Engine/Plugins/Animation/LiveLinkHub/Source/LiveLinkHub/Private/LiveLinkHubCommands.h b/Engine/Plugins/Animation/LiveLinkHub/Source/LiveLinkHub/Private/LiveLinkHubCommands.h
--- a/Engine/Plugins/Animation/LiveLinkHub/Source/LiveLinkHub/Private/LiveLinkHubCommands.h
+++ b/Engine/Plugins/Animation/LiveLinkHub/Source/LiveLinkHub/Private/LiveLinkHubCommands.h
@@ -38,6 +38,7 @@ public:
 		UI_COMMAND(OpenConfig, "Open Config...", "Open an existing config", EUserInterfaceActionType::Button, FInputChord(EKeys::O, EModifierKey::Control));
 		UI_COMMAND(SaveConfigAs, "Save Config As...", "Save a config to a new file", EUserInterfaceActionType::Button, FInputChord(EModifierKey::Control | EModifierKey::Alt, EKeys::S));
 		UI_COMMAND(SaveConfig, "Save Config", "Save the current config", EUserInterfaceActionType::Button, FInputChord(EKeys::S, EModifierKey::Control));
+		UI_COMMAND(OpenLogsFolder, "Open Logs Folder", "Open the folder containing the logs.", EUserInterfaceActionType::Button, FInputChord());
 	}
 
 public:
@@ -49,6 +50,8 @@ public:
 	TSharedPtr<FUICommandInfo> SaveConfigAs;
 	/** Save the current configuration to an existing file. */
 	TSharedPtr<FUICommandInfo> SaveConfig;
+	/** Open the folder containing the logs for live link hub. */
+	TSharedPtr<FUICommandInfo> OpenLogsFolder;
 };
 
 #undef LOCTEXT_NAMESPACE
\ No newline at end of file
diff --git a/Engine/Plugins/Animation/LiveLinkHub/Source/LiveLinkHub/Private/UI/Widgets/SLiveLinkHubTabViewWithManagerBase.cpp b/Engine/Plugins/Animation/LiveLinkHub/Source/LiveLinkHub/Private/UI/Widgets/SLiveLinkHubTabViewWithManagerBase.cpp
--- a/Engine/Plugins/Animation/LiveLinkHub/Source/LiveLinkHub/Private/UI/Widgets/SLiveLinkHubTabViewWithManagerBase.cpp
+++ b/Engine/Plugins/Animation/LiveLinkHub/Source/LiveLinkHub/Private/UI/Widgets/SLiveLinkHubTabViewWithManagerBase.cpp
@@ -2,27 +2,8 @@
 
 #include "SLiveLinkHubTabViewWithManagerBase.h"
 
-#include "Framework/Application/SlateApplication.h"
-#include "Framework/Commands/UICommandList.h"
 #include "Framework/Docking/LayoutService.h"
 #include "Framework/MultiBox/MultiBoxBuilder.h"
-#include "LiveLinkHub.h"
-#include "LiveLinkHubCommands.h"
-#include "LiveLinkHubModule.h"
-#include "Misc/CommandLine.h"
-#include "Misc/EngineBuildSettings.h"
-#include "Modules/ModuleManager.h"
-#include "UI/Widgets/SLiveLinkHubTabViewBase.h"
-
-#if !UE_BUILD_SHIPPING
-#	if WITH_UNREAL_TARGET_DEVELOPER_TOOLS
-#		include "ISessionFrontendModule.h"
-#	endif
-#	include "ISlateReflectorModule.h"
-#endif
-
-
-#define LOCTEXT_NAMESPACE "LiveLinkHub.SLiveLinkHubTabViewWithManagerBase"
 
 void SLiveLinkHubTabViewWithManagerBase::Construct(const FArguments& InArgs)
 {
@@ -57,89 +38,5 @@ TSharedRef<SWidget> SLiveLinkHubTabViewWithManagerBase::CreateTabs(const FArgume
 		)
 	);
 
-	TSharedRef<SWidget> Result = TabManager->RestoreFrom(Layout, InArgs._ConstructUnderWindow).ToSharedRef();
-	
-	const TSharedPtr<FLiveLinkHub> LiveLinkHub = FModuleManager::Get().GetModuleChecked<FLiveLinkHubModule>("LiveLinkHub").GetLiveLinkHub();
-	check(LiveLinkHub.IsValid());
-	FMenuBarBuilder MenuBarBuilder = FMenuBarBuilder(LiveLinkHub->GetCommandList());
-	FillInDefaultMenuItems(MenuBarBuilder);
-	InArgs._CreateMenuBar.ExecuteIfBound(MenuBarBuilder);
-	
-	const TSharedRef<SWidget> MenuWidget = MenuBarBuilder.MakeWidget();
-	TabManager->SetAllowWindowMenuBar(true);
-	TabManager->SetMenuMultiBox(MenuBarBuilder.GetMultiBox(), MenuWidget);
-	
-	return Result;
-}
-
-void SLiveLinkHubTabViewWithManagerBase::FillInDefaultMenuItems(FMenuBarBuilder MenuBarBuilder)
-{
-	MenuBarBuilder.AddPullDownMenu(
-		LOCTEXT("FileMenuLabel", "File"),
-		FText::GetEmpty(),
-		FNewMenuDelegate::CreateSP(this, &SLiveLinkHubTabViewWithManagerBase::FillFileMenu),
-		"File"
-	);
-	
-	MenuBarBuilder.AddPullDownMenu(
-		LOCTEXT("WindowMenuLabel", "Window"),
-		FText::GetEmpty(),
-		FNewMenuDelegate::CreateSP(this, &SLiveLinkHubTabViewWithManagerBase::FillWindowMenu),
-		"Window"
-	);
-}
-
-void SLiveLinkHubTabViewWithManagerBase::FillFileMenu(FMenuBuilder& MenuBuilder)
-{
-	MenuBuilder.BeginSection("Open", LOCTEXT("OpenHeader", "Open"));
-	MenuBuilder.AddMenuEntry(FLiveLinkHubCommands::Get().NewConfig);
-	MenuBuilder.AddMenuEntry(FLiveLinkHubCommands::Get().OpenConfig);
-	MenuBuilder.EndSection();
-	
-	MenuBuilder.BeginSection("Save", LOCTEXT("SaveHeader", "Save"));
-	MenuBuilder.AddMenuEntry(FLiveLinkHubCommands::Get().SaveConfig);
-	MenuBuilder.AddMenuEntry(FLiveLinkHubCommands::Get().SaveConfigAs);
-	MenuBuilder.EndSection();
-
-#if !UE_BUILD_SHIPPING
-	const bool bShowDevTools = FParse::Param(FCommandLine::Get(), TEXT("Development"))
-		|| FEngineBuildSettings::IsInternalBuild();
-
-	if (bShowDevTools)
-	{
-		MenuBuilder.BeginSection("Development", LOCTEXT("DevelopmentHeader", "Development"));
-#if WITH_UNREAL_TARGET_DEVELOPER_TOOLS
-		MenuBuilder.AddMenuEntry(
-			LOCTEXT("FileMenu_Development_AutomationTools", "Automation Tools"),
-			FText(),
-			FSlateIcon(),
-			FExecuteAction::CreateStatic([]()
-			{
-				ISessionFrontendModule& SessionFrontend =
-					FModuleManager::LoadModuleChecked<ISessionFrontendModule>("SessionFrontend");
-				SessionFrontend.InvokeSessionFrontend(FName("AutomationPanel"));
-			})
-		);
-#endif // #if WITH_UNREAL_TARGET_DEVELOPER_TOOLS
-		MenuBuilder.AddMenuEntry(
-			LOCTEXT("FileMenu_Development_WidgetReflector", "Widget Reflector"),
-			FText(),
-			FSlateIcon(),
-			FExecuteAction::CreateStatic([]()
-			{
-				ISlateReflectorModule& SlateReflector =
-					FModuleManager::LoadModuleChecked<ISlateReflectorModule>("SlateReflector");
-				SlateReflector.DisplayWidgetReflector();
-			})
-		);
-		MenuBuilder.EndSection();
-	}
-#endif // #if !UE_BUILD_SHIPPING
-}
-
-void SLiveLinkHubTabViewWithManagerBase::FillWindowMenu(FMenuBuilder& MenuBuilder)
-{
-	TabManager->PopulateLocalTabSpawnerMenu(MenuBuilder);
+	return TabManager->RestoreFrom(Layout, InArgs._ConstructUnderWindow).ToSharedRef();
 }
-
-#undef LOCTEXT_NAMESPACE
diff --git a/Engine/Plugins/Animation/LiveLinkHub/Source/LiveLinkHub/Private/UI/Widgets/SLiveLinkHubTabViewWithManagerBase.h b/Engine/Plugins/Animation/LiveLinkHub/Source/LiveLinkHub/Private/UI/Widgets/SLiveLinkHubTabViewWithManagerBase.h
--- a/Engine/Plugins/Animation/LiveLinkHub/Source/LiveLinkHub/Private/UI/Widgets/SLiveLinkHubTabViewWithManagerBase.h
+++ b/Engine/Plugins/Animation/LiveLinkHub/Source/LiveLinkHub/Private/UI/Widgets/SLiveLinkHubTabViewWithManagerBase.h
@@ -2,9 +2,10 @@
 
 #pragma once
 
-#include "Framework/Docking/TabManager.h"
 #include "UI/Widgets/SLiveLinkHubTabViewBase.h"
 
+#include "Framework/Docking/TabManager.h"
+
 class FMenuBarBuilder;
 
 /** Base class for tab views that create sub-tabs. */
@@ -23,8 +24,6 @@ public:
 		SLATE_ARGUMENT(TSharedPtr<SWindow>, ConstructUnderWindow)
 		/** Callback for creating the sub-tabs */
 		SLATE_EVENT(FCreateTabs, CreateTabs)
-		/** Called to fill the major tab's menu bar */
-		SLATE_EVENT(FCreateWindowMenu, CreateMenuBar)
 		/** Optional function to overlay all the tabs (for example to overlay a message to execute a console command first) */
 		SLATE_EVENT(FOverlayTabs, OverlayTabs)
 		/** Name to give the layout. Important for saving config. */
@@ -41,15 +40,8 @@ protected:
 	TSharedPtr<FTabManager> GetTabManager() const { return TabManager; }
 	
 private:
-
 	/** Holds the tab manager that manages the front-end's tabs. */
 	TSharedPtr<FTabManager> TabManager;
 	/** Handles restoring the layout and instantiating tabs. */
 	TSharedRef<SWidget> CreateTabs(const FArguments& InArgs);
-	/** Creates the top bar menu items. */
-	void FillInDefaultMenuItems(FMenuBarBuilder MenuBarBuilder);
-	/** Populates the tab spawner's file menu. */
-	void FillFileMenu(FMenuBuilder& MenuBuilder);
-	/** Populates the tab spawner's window menu. */
-	void FillWindowMenu(FMenuBuilder& MenuBuilder);
 };
