LiveLinkHub - Clean up dead code now that layouts are implemented

#rnx

diff --git a/Engine/Plugins/Animation/LiveLinkDevice/Source/LiveLinkDevice/Private/LiveLinkDeviceModule.cpp b/Engine/Plugins/Animation/LiveLinkDevice/Source/LiveLinkDevice/Private/LiveLinkDeviceModule.cpp
--- a/Engine/Plugins/Animation/LiveLinkDevice/Source/LiveLinkDevice/Private/LiveLinkDeviceModule.cpp
+++ b/Engine/Plugins/Animation/LiveLinkDevice/Source/LiveLinkDevice/Private/LiveLinkDeviceModule.cpp
@@ -3,10 +3,8 @@
 #include "LiveLinkDeviceModule.h"
 #include "DetailsViewArgs.h"
 #include "Engine/Engine.h"
-#include "Features/IModularFeatures.h"
 #include "Framework/Docking/TabManager.h"
 #include "Framework/MultiBox/MultiBoxBuilder.h"
-#include "ILiveLinkHubComponent.h"
 #include "LiveLinkDevice.h"
 #include "LiveLinkDeviceSubsystem.h"
 #include "Logging/StructuredLog.h"
@@ -29,17 +27,6 @@ const FName FLiveLinkDeviceModule::DeviceDetailsTabName = FName("LiveLinkDeviceD
 IMPLEMENT_MODULE(FLiveLinkDeviceModule, LiveLinkDevice);
 
 
-class FLiveLinkHubDeviceComponent : public ILiveLinkHubComponent
-{
-public:
-	virtual void Init(const FLiveLinkHubComponentInitParams& InParams) override
-	{
-		InParams.MainStack->AddTab(FLiveLinkDeviceModule::DevicesTabName, ETabState::OpenedTab);
-		InParams.MainStack->AddTab(FLiveLinkDeviceModule::DeviceDetailsTabName, ETabState::OpenedTab);
-	}
-};
-
-
 //////////////////////////////////////////////////////////////////////////
 
 
@@ -54,9 +41,6 @@ void FLiveLinkDeviceModule::StartupModule()
 		->RegisterNomadTabSpawner(DeviceDetailsTabName,
 			FOnSpawnTab::CreateRaw(this, &FLiveLinkDeviceModule::OnSpawnDeviceDetailsTab))
 		.SetDisplayName(LOCTEXT("DeviceDetailsTabDisplayName", "Device Details"));
-
-	LLHComponent = MakeShared<FLiveLinkHubDeviceComponent>();
-	IModularFeatures::Get().RegisterModularFeature(ILiveLinkHubComponent::ModularFeatureName, LLHComponent.Get());
 }
 
 
diff --git a/Engine/Plugins/Animation/LiveLinkDevice/Source/LiveLinkDevice/Private/LiveLinkDeviceModule.h b/Engine/Plugins/Animation/LiveLinkDevice/Source/LiveLinkDevice/Private/LiveLinkDeviceModule.h
--- a/Engine/Plugins/Animation/LiveLinkDevice/Source/LiveLinkDevice/Private/LiveLinkDeviceModule.h
+++ b/Engine/Plugins/Animation/LiveLinkDevice/Source/LiveLinkDevice/Private/LiveLinkDeviceModule.h
@@ -8,7 +8,6 @@
 #include "UObject/WeakObjectPtr.h"
 
 
-class FLiveLinkHubDeviceComponent;
 class FSpawnTabArgs;
 class IDetailsView;
 class SDockTab;
@@ -37,7 +36,6 @@ private:
 private:
 	void DeviceSelectionChanged(ULiveLinkDevice* InSelectedDevice);
 
-	TSharedPtr<FLiveLinkHubDeviceComponent> LLHComponent;
 	TSharedPtr<SLiveLinkDeviceTable> DeviceTable;
 	TSharedPtr<IDetailsView> DetailsView;
 	TWeakObjectPtr<ULiveLinkDevice> WeakSelectedDevice;
diff --git a/Engine/Plugins/Animation/LiveLinkHub/Source/LiveLinkHub/Private/UI/Widgets/LiveLinkHubMainTabController.cpp b/Engine/Plugins/Animation/LiveLinkHub/Source/LiveLinkHub/Private/UI/Widgets/LiveLinkHubMainTabController.cpp
deleted file mode 100644
diff --git a/Engine/Plugins/Animation/LiveLinkHub/Source/LiveLinkHub/Private/UI/Widgets/LiveLinkHubMainTabController.h b/Engine/Plugins/Animation/LiveLinkHub/Source/LiveLinkHub/Private/UI/Widgets/LiveLinkHubMainTabController.h
deleted file mode 100644
diff --git a/Engine/Plugins/Animation/LiveLinkHub/Source/LiveLinkHub/Private/UI/Widgets/SLiveLinkHubMainTabView.cpp b/Engine/Plugins/Animation/LiveLinkHub/Source/LiveLinkHub/Private/UI/Widgets/SLiveLinkHubMainTabView.cpp
deleted file mode 100644
diff --git a/Engine/Plugins/Animation/LiveLinkHub/Source/LiveLinkHub/Private/UI/Widgets/SLiveLinkHubMainTabView.h b/Engine/Plugins/Animation/LiveLinkHub/Source/LiveLinkHub/Private/UI/Widgets/SLiveLinkHubMainTabView.h
deleted file mode 100644
diff --git a/Engine/Plugins/Animation/LiveLinkHub/Source/LiveLinkHub/Private/UI/Widgets/SLiveLinkHubStatusBar.cpp b/Engine/Plugins/Animation/LiveLinkHub/Source/LiveLinkHub/Private/UI/Widgets/SLiveLinkHubStatusBar.cpp
deleted file mode 100644
diff --git a/Engine/Plugins/Animation/LiveLinkHub/Source/LiveLinkHub/Private/UI/Widgets/SLiveLinkHubStatusBar.h b/Engine/Plugins/Animation/LiveLinkHub/Source/LiveLinkHub/Private/UI/Widgets/SLiveLinkHubStatusBar.h
deleted file mode 100644
diff --git a/Engine/Plugins/Animation/LiveLinkHub/Source/LiveLinkHub/Private/UI/Widgets/SLiveLinkHubTabViewBase.cpp b/Engine/Plugins/Animation/LiveLinkHub/Source/LiveLinkHub/Private/UI/Widgets/SLiveLinkHubTabViewBase.cpp
deleted file mode 100644
diff --git a/Engine/Plugins/Animation/LiveLinkHub/Source/LiveLinkHub/Private/UI/Widgets/SLiveLinkHubTabViewBase.h b/Engine/Plugins/Animation/LiveLinkHub/Source/LiveLinkHub/Private/UI/Widgets/SLiveLinkHubTabViewBase.h
deleted file mode 100644
diff --git a/Engine/Plugins/Animation/LiveLinkHub/Source/LiveLinkHub/Private/UI/Widgets/SLiveLinkHubTabViewWithManagerBase.cpp b/Engine/Plugins/Animation/LiveLinkHub/Source/LiveLinkHub/Private/UI/Widgets/SLiveLinkHubTabViewWithManagerBase.cpp
deleted file mode 100644
diff --git a/Engine/Plugins/Animation/LiveLinkHub/Source/LiveLinkHub/Private/UI/Widgets/SLiveLinkHubTabViewWithManagerBase.h b/Engine/Plugins/Animation/LiveLinkHub/Source/LiveLinkHub/Private/UI/Widgets/SLiveLinkHubTabViewWithManagerBase.h
deleted file mode 100644
diff --git a/Engine/Plugins/Animation/LiveLinkHub/Source/LiveLinkHub/Private/UI/Window/LiveLinkHubWindowController.cpp b/Engine/Plugins/Animation/LiveLinkHub/Source/LiveLinkHub/Private/UI/Window/LiveLinkHubWindowController.cpp
--- a/Engine/Plugins/Animation/LiveLinkHub/Source/LiveLinkHub/Private/UI/Window/LiveLinkHubWindowController.cpp
+++ b/Engine/Plugins/Animation/LiveLinkHub/Source/LiveLinkHub/Private/UI/Window/LiveLinkHubWindowController.cpp
@@ -11,7 +11,6 @@
 #include "LiveLinkHub.h"
 #include "HAL/PlatformApplicationMisc.h"
 #include "Misc/ConfigCacheIni.h"
-#include "UI/Widgets/LiveLinkHubMainTabController.h"
 #include "UI/Widgets/SLiveLinkHubMemoryStats.h"
 #include "UI/Widgets/SLiveLinkHubTopologyModeSwitcher.h"
 #include "UI/Window/ModalWindowManager.h"
@@ -20,7 +19,6 @@
 #define LOCTEXT_NAMESPACE "LiveLinkHubWindowController"
 
 FLiveLinkHubWindowController::FLiveLinkHubWindowController()
-	: MainTabController(MakeShared<FLiveLinkHubMainTabController>())
 {
 	ModalWindowManager = InitializeSlateApplication();
 }
@@ -144,18 +142,6 @@ TSharedPtr<FModalWindowManager> FLiveLinkHubWindowController::InitializeSlateApp
 	return MakeShared<FModalWindowManager>(CreateWindow());
 }
 
-void FLiveLinkHubWindowController::InitComponents(const TSharedRef<FTabManager::FStack>& MainArea)
-{
-	const FLiveLinkHubComponentInitParams Params { GetRootWindow().ToSharedRef(), MainArea };
-	MainTabController->Init(Params);
-
-	TArray<ILiveLinkHubComponent*> ExternalComponents = IModularFeatures::Get().GetModularFeatureImplementations<ILiveLinkHubComponent>(ILiveLinkHubComponent::ModularFeatureName);
-	for (ILiveLinkHubComponent* LiveLinkHubComponent : ExternalComponents)
-	{
-		LiveLinkHubComponent->Init(Params);
-	}
-}
-
 void FLiveLinkHubWindowController::OnWindowClosed(const TSharedRef<SWindow>& Window)
 {
 	if (TSharedPtr<FLiveLinkHub> LiveLinkHub = FLiveLinkHub::Get())
diff --git a/Engine/Plugins/Animation/LiveLinkHub/Source/LiveLinkHub/Private/UI/Window/LiveLinkHubWindowController.h b/Engine/Plugins/Animation/LiveLinkHub/Source/LiveLinkHub/Private/UI/Window/LiveLinkHubWindowController.h
--- a/Engine/Plugins/Animation/LiveLinkHub/Source/LiveLinkHub/Private/UI/Window/LiveLinkHubWindowController.h
+++ b/Engine/Plugins/Animation/LiveLinkHub/Source/LiveLinkHub/Private/UI/Window/LiveLinkHubWindowController.h
@@ -5,11 +5,9 @@
 #include "Framework/Docking/TabManager.h"
 
 class FAssetEditorToolkit;
-class FLiveLinkHubMainTabController;
 class FLiveLinkHubSessionTabBase;
 class FOutputLogController;
 class ILiveLinkClient;
-class ILiveLinkHubComponent;
 class FModalWindowManager;
 class SWindow;
 
@@ -30,8 +28,6 @@ private:
 	TSharedRef<SWindow> CreateWindow();
 	/** Create the slate application that hosts the livelink hub. */
 	TSharedPtr<FModalWindowManager> InitializeSlateApplication();
-	/** Initialize components. */
-	void InitComponents(const TSharedRef<FTabManager::FStack>& MainArea);
 	/** Window closed handler. */
 	void OnWindowClosed(const TSharedRef<SWindow>& Window);
 private:
@@ -43,8 +39,6 @@ private:
 	TSharedPtr<FTabManager::FLayout> PersistentLayout;
 	/** The main window being managed */
 	TSharedPtr<SWindow> RootWindow;
-	/** Manages the main tab. */
-	TSharedRef<FLiveLinkHubMainTabController> MainTabController;
 	/** Manages modal windows for the application. */
 	TSharedPtr<FModalWindowManager> ModalWindowManager;
 	/** Menu bar widget for the hub. */
diff --git a/Engine/Plugins/Animation/LiveLinkHub/Source/LiveLinkHub/Public/ILiveLinkHubComponent.h b/Engine/Plugins/Animation/LiveLinkHub/Source/LiveLinkHub/Public/ILiveLinkHubComponent.h
--- a/Engine/Plugins/Animation/LiveLinkHub/Source/LiveLinkHub/Public/ILiveLinkHubComponent.h
+++ b/Engine/Plugins/Animation/LiveLinkHub/Source/LiveLinkHub/Public/ILiveLinkHubComponent.h
@@ -8,7 +8,7 @@
 
 class SWindow;
 
-struct LIVELINKHUB_API FLiveLinkHubComponentInitParams
+struct UE_DEPRECATED(5.6, "Customize LiveLinkHub by adding layouts instead.") LIVELINKHUB_API FLiveLinkHubComponentInitParams
 {
 	/** LiveLinkHub's main Window. */
 	TSharedRef<SWindow> Window;
@@ -24,7 +24,7 @@ struct LIVELINKHUB_API FLiveLinkHubComponentInitParams
 
 
 /** Provides an interface for elements that wish to add additional content (tabs) to LiveLinkHub. */
-class LIVELINKHUB_API ILiveLinkHubComponent : public IModularFeature
+class UE_DEPRECATED(5.6, "Customize LiveLinkHub by adding layouts instead.") LIVELINKHUB_API ILiveLinkHubComponent : public IModularFeature
 {
 public:
 	/** Name of the modular feature. */
@@ -32,6 +32,8 @@ public:
 
 	virtual ~ILiveLinkHubComponent() = default;
 
+PRAGMA_DISABLE_DEPRECATION_WARNINGS
 	/** Initialises the component, e.g. registering tab spawners, etc. */
 	virtual void Init(const FLiveLinkHubComponentInitParams& Params) = 0;
+PRAGMA_ENABLE_DEPRECATION_WARNINGS
 };
