Live Link Hub - Allow users to create, load, export and reset layouts

User layouts can be created by modifying your current layout, then clicking Save Layout As. This will add the save path to the LLH user settings so they can be re-discovered on the next launch.

Only default(code-defined) layouts are allowed to be reset, since we're just reverting the layout to its code-defined version.

diff --git a/Engine/Plugins/Animation/LiveLinkHub/Source/LiveLinkHub/Private/LiveLinkHub.cpp b/Engine/Plugins/Animation/LiveLinkHub/Source/LiveLinkHub/Private/LiveLinkHub.cpp
--- a/Engine/Plugins/Animation/LiveLinkHub/Source/LiveLinkHub/Private/LiveLinkHub.cpp
+++ b/Engine/Plugins/Animation/LiveLinkHub/Source/LiveLinkHub/Private/LiveLinkHub.cpp
@@ -179,6 +179,8 @@ void FLiveLinkHub::InitializePostEngineInit()
 
 	SetCurrentMode("CreatorMode");
 
+	DiscoverLayouts();
+
 	LiveLinkHubClient->OnStaticDataReceived_AnyThread().AddSP(this, &FLiveLinkHub::OnStaticDataReceived_AnyThread);
 	LiveLinkHubClient->OnFrameDataReceived_AnyThread().AddSP(this, &FLiveLinkHub::OnFrameDataReceived_AnyThread);
 	LiveLinkHubClient->OnSubjectMarkedPendingKill_AnyThread().AddSP(this, &FLiveLinkHub::OnSubjectMarkedPendingKill_AnyThread);
diff --git a/Engine/Plugins/Animation/LiveLinkHub/Source/LiveLinkHub/Private/LiveLinkHubApplication.cpp b/Engine/Plugins/Animation/LiveLinkHub/Source/LiveLinkHub/Private/LiveLinkHubApplication.cpp
--- a/Engine/Plugins/Animation/LiveLinkHub/Source/LiveLinkHub/Private/LiveLinkHubApplication.cpp
+++ b/Engine/Plugins/Animation/LiveLinkHub/Source/LiveLinkHub/Private/LiveLinkHubApplication.cpp
@@ -2,8 +2,15 @@
 
 #include "LiveLinkHubApplication.h"
 
+#include "Algo/RemoveIf.h"
+#include "Dom/JsonObject.h"
+#include "HAL/FileManagerGeneric.h"
 #include "LiveLinkHubCreatorAppMode.h"
+#include "LiveLinkHubUserLayoutMode.h"
 #include "Misc/EngineBuildSettings.h"
+#include "Misc/FileHelper.h"
+#include "Misc/ScopeExit.h"
+#include "Serialization/JsonSerializer.h"
 #include "StatusBarSubsystem.h"
 #include "Toolkits/AssetEditorToolkitMenuContext.h"
 #include "ToolMenus.h"
@@ -20,11 +27,29 @@
 
 #define LOCTEXT_NAMESPACE "LiveLinkHubApplication"
 
+struct FLiveLinkHubUserLayout
+{
+	FString ParentModeName;
+	TSharedPtr<FJsonObject> JsonLayout;
+};
+
+namespace LiveLinkHubLayoutFileUtils
+{
+	const FString LayoutFileExtension = TEXT("llhlayout");
+	static const FString LayoutDescription = TEXT("Live Link Hub Layout");
+	static const FString LayoutExtension = LayoutFileExtension;
+	static const FString LayoutDefaultFileName = TEXT("New Layout");
+
+	static const FString FileTypes = FString::Printf(TEXT("%s (*.%s)|*.%s"), *LayoutDescription, *LayoutExtension, *LayoutExtension);
+}
+
 FLiveLinkHubApplication::FLiveLinkHubApplication()
 	: FLiveLinkHubApplicationBase()
 {
 	LayoutIni = TEXT("LiveLinkHubLayout");
 
+	FEditorDirectories::Get().LoadLastDirectories();
+
 	// We set an extender here and collect toolbar widgets 
 	TSharedPtr<FExtender> Extender = MakeShared<FExtender>();
 	Extender->AddToolBarExtension(
@@ -52,13 +77,19 @@ void FLiveLinkHubApplication::AddToolbarExtenders(FToolBarBuilder&)
 
 void FLiveLinkHubApplication::AddLiveLinkHubApplicationMode(FName ModeName, TSharedRef<FLiveLinkHubApplicationMode> Mode)
 {
-	CachedModeInfo.Add(ModeName, FLiveLinkHubAppModeInfo{ Mode->GetModeIcon(), Mode->GetDisplayName()});
+	CachedModeInfo.Add(ModeName, FLiveLinkHubAppModeInfo{ Mode->GetModeIcon(), Mode->GetDisplayName(), Mode->IsUserLayout()});
 
 	TabManager->AddLocalWorkspaceMenuItem(Mode->GetWorkspaceMenuCategory());
 
 	AddApplicationMode(ModeName, MoveTemp(Mode));
 }
 
+void FLiveLinkHubApplication::RemoveLiveLinkHubApplicationMode(FName ModeName)
+{
+	CachedModeInfo.Remove(ModeName);
+	RemoveApplicationMode(ModeName);
+}
+
 void FLiveLinkHubApplication::PushTabFactories(const FWorkflowAllowedTabSet& FactorySetToPush, TSharedPtr<FLiveLinkHubApplicationMode> ApplicationMode)
 {
 	check(TabManager.IsValid());
@@ -122,6 +153,219 @@ void FLiveLinkHubApplication::SetCurrentMode(FName NewMode)
 	AppModeChangedDelegate.Broadcast(NewMode);
 }
 
+void FLiveLinkHubApplication::PersistUserLayout(const FString& LayoutName, const TSharedPtr<FJsonObject>& JsonLayout)
+{
+	if (FString* LayoutPath = CachedLayouts.Find(LayoutName))
+	{
+		SaveLayoutToFile(*LayoutPath, JsonLayout);
+	}
+}
+
+void FLiveLinkHubApplication::SaveLayoutAs()
+{
+	if (TSharedPtr<FLiveLinkHub> LiveLinkHub = FLiveLinkHub::Get())
+	{
+		const FString DefaultFile = LiveLinkHubLayoutFileUtils::LayoutDefaultFileName;
+
+		TArray<FString> SaveFileNames;
+
+		IDesktopPlatform* DesktopPlatform = FDesktopPlatformModule::Get();
+		const void* ParentWindowWindowHandle = FSlateApplication::Get().FindBestParentWindowHandleForDialogs(nullptr);
+
+		const FString DefaultSaveDir = FEditorDirectories::Get().GetLastDirectory(ELastDirectory::GENERIC_SAVE);
+
+		const bool bFileSelected = DesktopPlatform->SaveFileDialog(
+			ParentWindowWindowHandle,
+			LOCTEXT("LiveLinkHubSaveAsTitle", "Save As").ToString(),
+			DefaultSaveDir,
+			DefaultFile,
+			LiveLinkHubLayoutFileUtils::FileTypes,
+			EFileDialogFlags::None,
+			SaveFileNames);
+
+		if (bFileSelected && SaveFileNames.Num() > 0)
+		{
+			const FString& SavePath = SaveFileNames[0];
+
+			// Todo: Allow naming layouts differently from the filename? 
+			FString LayoutName = FPaths::GetBaseFilename(SavePath);
+			SaveLayoutToFile(SaveFileNames[0]);
+		
+			CachedLayouts.Add(LayoutName, SavePath);
+
+			// Todo: We already have the json object in memory, we could fast track this instead of saving to disk then reloading
+			RegisterUserLayout(LayoutName, SavePath);
+
+			ULiveLinkHubUserSettings* Settings = GetMutableDefault<ULiveLinkHubUserSettings>();
+			Settings->LayoutDirectories.AddUnique(FPaths::GetPath(SavePath));
+			Settings->SaveConfig();
+
+			SetCurrentMode(*LayoutName);
+		}
+	}
+}
+
+void FLiveLinkHubApplication::LoadLayout()
+{
+	const FString DefaultFile = LiveLinkHubLayoutFileUtils::LayoutDefaultFileName;
+
+	TArray<FString> OpenFileNames;
+
+	IDesktopPlatform* DesktopPlatform = FDesktopPlatformModule::Get();
+	const void* ParentWindowWindowHandle = FSlateApplication::Get().FindBestParentWindowHandleForDialogs(nullptr);
+
+	const bool bFileSelected = DesktopPlatform->OpenFileDialog(
+		ParentWindowWindowHandle,
+		LOCTEXT("LiveLinkHubOpenTitle", "Open").ToString(),
+		FEditorDirectories::Get().GetLastDirectory(ELastDirectory::GENERIC_OPEN),
+		DefaultFile,
+		LiveLinkHubLayoutFileUtils::LayoutDefaultFileName,
+		EFileDialogFlags::None,
+		OpenFileNames);
+
+	if (bFileSelected && OpenFileNames.Num() > 0)
+	{
+		const FString LayoutPath = OpenFileNames[0];
+		ULiveLinkHubUserSettings* Settings = GetMutableDefault<ULiveLinkHubUserSettings>();
+		Settings->LayoutDirectories.AddUnique(FPaths::GetPath(LayoutPath));
+		Settings->SaveConfig();
+
+		FString LayoutName = FPaths::GetBaseFilename(LayoutPath);
+
+		if (!CachedLayouts.Contains(LayoutName))
+		{
+			if (RegisterUserLayout(LayoutName, LayoutPath))
+			{
+				SetCurrentMode(*LayoutName);
+			}
+		}
+		else
+		{
+			UE_LOG(LogLiveLinkHub, Warning, TEXT("Could not load layout %s since it already exists."), *LayoutName);
+		}
+	}
+}
+
+void FLiveLinkHubApplication::ResetLayout()
+{
+	if (TSharedPtr<FLiveLinkHubApplicationMode> CurrentMode = StaticCastSharedPtr<FLiveLinkHubApplicationMode>(GetCurrentModePtr()))
+	{
+		constexpr bool bLoadUserLayout = false;
+		RestoreFromLayout(CurrentMode->GetTabLayout().ToSharedRef(), bLoadUserLayout);
+		AddToolbarExtender(CurrentMode->GetToolbarExtender());
+		RegenerateMenusAndToolbars();
+	}
+}
+
+TArray<FString> FLiveLinkHubApplication::GetUserLayouts()
+{
+	TArray<FString> Keys;
+	CachedLayouts.GenerateKeyArray(Keys);
+	Keys.Sort();
+	return Keys;
+}
+
+void FLiveLinkHubApplication::DeleteUserLayout(const FString& LayoutName)
+{
+	if (FString* LayoutPath = CachedLayouts.Find(LayoutName))
+	{
+		IFileManager::Get().Delete(**LayoutPath);
+
+		CachedLayouts.Remove(LayoutName);
+		RemoveLiveLinkHubApplicationMode(*LayoutName);
+	}
+}
+
+TOptional<FLiveLinkHubUserLayout> FLiveLinkHubApplication::ParseUserLayout(const FString& LayoutPath)
+{
+	TOptional<FLiveLinkHubUserLayout> UserLayout;
+
+	FString Result;
+	FFileHelper::LoadFileToString(Result, *LayoutPath);
+
+	TSharedRef<TJsonReader<TCHAR>> Reader = TJsonReaderFactory<TCHAR>::Create(Result);
+
+	TSharedPtr<FJsonObject> RootObject = MakeShared<FJsonObject>();
+	if (FJsonSerializer::Deserialize(Reader, RootObject))
+	{
+		if (RootObject->HasField(TEXT("ParentMode")) && RootObject->HasField(TEXT("Layout")))
+		{
+			FLiveLinkHubUserLayout Layout = FLiveLinkHubUserLayout{ RootObject->GetField(TEXT("ParentMode"), EJson::String)->AsString(), RootObject->GetObjectField(TEXT("Layout")) };
+			UserLayout = MoveTemp(Layout);
+		}
+		else
+		{
+			UE_LOG(LogLiveLinkHub, Error, TEXT("Live Link Hub Layout missing ParentNode field."));
+		}
+	}
+	else
+	{
+		UE_LOG(LogLiveLinkHub, Error, TEXT("Failed to parse livelinkhub layout."));
+	}
+
+	return UserLayout;
+}
+
+TSharedPtr<FLiveLinkHubApplicationMode> FLiveLinkHubApplication::FindApplicationMode(const FString& ModeName)
+{
+	const TMap<FName, TSharedPtr<FApplicationMode>>& ModeList = GetApplicationModeList();
+	ensureMsgf(CachedModeInfo.Contains(*ModeName), TEXT("Mode was not registered with LiveLinkHubApplication."));
+	if (TSharedPtr<FApplicationMode> AppMode = ModeList.FindRef(*ModeName))
+	{
+		return StaticCastSharedPtr<FLiveLinkHubApplicationMode>(AppMode);
+	}
+
+	return nullptr;
+}
+
+void FLiveLinkHubApplication::SaveLayoutToFile(const FString& SavePath, TSharedPtr<FJsonObject> JsonLayout)
+{
+	if (SavePath.IsEmpty() && LastLayoutPath.IsEmpty())
+	{
+		return;
+	}
+
+	if (!SavePath.IsEmpty())
+	{
+		LastLayoutPath = SavePath;
+	}
+
+	FEditorDirectories::Get().SetLastDirectory(ELastDirectory::GENERIC_SAVE, FPaths::GetPath(LastLayoutPath));
+	FEditorDirectories::Get().SaveLastDirectories();
+
+	if (!JsonLayout)
+	{
+		TSharedRef<FTabManager::FLayout> PersistentLayout = FLiveLinkHub::Get()->GetTabManager()->PersistLayout();
+		JsonLayout = PersistentLayout->ToJson();
+	}
+
+	if (ensure(!LastLayoutPath.IsEmpty()))
+	{
+		TSharedRef<FJsonObject> LayoutJson = MakeShared<FJsonObject>();
+
+		FString ParentModeName = GetCurrentMode().ToString();
+		if (TSharedPtr<FLiveLinkHubApplicationMode> ParentMode = FindApplicationMode(ParentModeName))
+		{
+			if (ParentMode->IsUserLayout())
+			{
+				ParentModeName = StaticCastSharedPtr<FLiveLinkHubUserLayoutMode>(ParentMode)->GetParentMode()->GetModeName().ToString();
+			}
+		}
+		LayoutJson->SetStringField(TEXT("ParentMode"), ParentModeName);
+		LayoutJson->SetObjectField(TEXT("Layout"), JsonLayout);
+
+		FString LayoutAsString;
+		TSharedRef<TJsonWriter<>> Writer = TJsonWriterFactory<>::Create( &LayoutAsString );
+		if (!FJsonSerializer::Serialize(LayoutJson, Writer))
+		{
+			UE_LOG(LogLiveLinkHub, Error, TEXT("Failed save layout as Json string"));
+			return;
+		}
+
+		FFileHelper::SaveStringToFile(LayoutAsString, *LastLayoutPath);
+	}
+}
+
 TOptional<FLiveLinkHubAppModeInfo> FLiveLinkHubApplication::GetModeInfo(FName Mode) const
 {
 	TOptional<FLiveLinkHubAppModeInfo> ModeInfo;
@@ -134,10 +378,22 @@ TOptional<FLiveLinkHubAppModeInfo> FLiveLinkHubApplication::GetModeInfo(FName Mo
 	return ModeInfo;
 }
 
+FName FLiveLinkHubApplication::GetLayoutName(FName ModeName)
+{
+	FName Name = NAME_None;
+	if (TSharedPtr<FApplicationMode> AppMode = GetApplicationModeList().FindRef(ModeName))
+	{
+		Name = StaticCastSharedPtr<FLiveLinkHubApplicationMode>(AppMode)->GetLayoutName();
+	}
+
+	return Name;
+}
+
 TArray<FName> FLiveLinkHubApplication::GetApplicationModes() const
 {
 	TArray<FName> ApplicationModeNames;
 	CachedModeInfo.GenerateKeyArray(ApplicationModeNames);
+	ApplicationModeNames.SetNum(Algo::RemoveIf(ApplicationModeNames, [this](FName Name) { return CachedLayouts.Contains(Name.ToString()); }));
 	return ApplicationModeNames;
 }
 
@@ -310,4 +566,52 @@ FText FLiveLinkHubApplication::GetToolkitName() const
 	return LOCTEXT("LiveLinkHubLabel", "Live Link Hub");
 }
 
+void FLiveLinkHubApplication::DiscoverLayouts()
+{
+	TArray<FString> Files;
+
+	for (const FString& Directory : GetLayoutDirectories())
+	{
+		IFileManager::Get().FindFiles(Files, *Directory, *LiveLinkHubLayoutFileUtils::LayoutFileExtension);
+
+		for (const FString& File : Files)
+		{
+			FString LayoutName = File.LeftChop(LiveLinkHubLayoutFileUtils::LayoutFileExtension.Len() + 1);
+
+			if (!CachedLayouts.Contains(LayoutName))
+			{
+				const FString LayoutFile = Directory / File;
+				RegisterUserLayout(LayoutName, LayoutFile);
+			}
+		}
+	}
+}
+
+bool FLiveLinkHubApplication::RegisterUserLayout(const FString& LayoutName, const FString& LayoutPath)
+{
+	if (TOptional<FLiveLinkHubUserLayout> UserLayout = ParseUserLayout(LayoutPath))
+	{
+		CachedLayouts.Add(LayoutName, LayoutPath);
+		if (TSharedPtr<FLiveLinkHubApplicationMode> ParentMode = FindApplicationMode(UserLayout->ParentModeName))
+		{
+			TSharedRef<FLiveLinkHubUserLayoutMode> UserLayoutMode = MakeShared<FLiveLinkHubUserLayoutMode>(*LayoutName, UserLayout->JsonLayout.ToSharedRef(), ParentMode);
+			AddLiveLinkHubApplicationMode(*LayoutName, UserLayoutMode);
+		}
+		return true;
+	}
+
+	return false;
+}
+
+TArray<FString> FLiveLinkHubApplication::GetLayoutDirectories() const
+{
+	FString UserSettingsDir = FLiveLinkHubApplication::UserSettingsDir() / TEXT("Layouts");
+	FString LastSaveDirectory = FEditorDirectories::Get().GetLastDirectory(ELastDirectory::GENERIC_SAVE);
+
+	TArray<FString> AllPaths = { MoveTemp(UserSettingsDir), MoveTemp(LastSaveDirectory) };
+	AllPaths.Append(GetDefault<ULiveLinkHubUserSettings>()->LayoutDirectories);
+
+	return AllPaths;
+}
+
 #undef LOCTEXT_NAMESPACE /*LiveLinkHubApplication*/
diff --git a/Engine/Plugins/Animation/LiveLinkHub/Source/LiveLinkHub/Private/LiveLinkHubApplication.h b/Engine/Plugins/Animation/LiveLinkHub/Source/LiveLinkHub/Private/LiveLinkHubApplication.h
--- a/Engine/Plugins/Animation/LiveLinkHub/Source/LiveLinkHub/Private/LiveLinkHubApplication.h
+++ b/Engine/Plugins/Animation/LiveLinkHub/Source/LiveLinkHub/Private/LiveLinkHubApplication.h
@@ -12,8 +12,12 @@ class FLiveLinkHubApplicationMode;
 /** Cached information about a given application mode. Used to populate the mode switcher widget. */
 struct FLiveLinkHubAppModeInfo
 {
+	/** Icon of the mode. */
 	FSlateIcon Icon;
+	/** Display name for the mode. */
 	FText DisplayName;
+	/** Whether this is a user layout mode. */
+	bool bUserLayout = false;
 };
 
 DECLARE_MULTICAST_DELEGATE_OneParam(FOnApplicationModeChanged, /*CurrentMode*/FName);
@@ -27,7 +31,10 @@ public:
 	FOnApplicationModeChanged& OnApplicationModeChanged() { return AppModeChangedDelegate; }
 
 	/** Get information about a given mode. Returns an empty optional if the mode wasn't found or if it wasn't added to the app using AddLiveLinkHubApplicationMode. */
-	TOptional<FLiveLinkHubAppModeInfo> GetModeInfo(FName Mode) const;
+	TOptional<FLiveLinkHubAppModeInfo> GetModeInfo(FName ModeName) const;
+
+	/** Get the name of the TabLayout for a given mode. */
+	FName GetLayoutName(FName ModeName);
 	
 	/** Get the list of modes that were registered with the app. */
 	TArray<FName> GetApplicationModes() const;
@@ -38,11 +45,36 @@ public:
 	/** Push tab factories for the given mode. */
 	virtual void PushTabFactories(const FWorkflowAllowedTabSet& FactorySetToPush, TSharedPtr<FLiveLinkHubApplicationMode> ApplicationMode);
 
+	/** Save a user layout to disk. */
+	void PersistUserLayout(const FString& LayoutName, const TSharedPtr<FJsonObject>& JsonLayout);
+
+	/** Prompt the user to save the current layout file to disk. */
+	void SaveLayoutAs();
+
+	/**  Prompt the user to load a layout file from disk.*/
+	void LoadLayout();
+
+	/** Reset the current layout to its code-defined version. (Only available for non-user layouts) */
+	void ResetLayout();
+
+	/** Get the list of user layouts. */
+	TArray<FString> GetUserLayouts();
+
+	/** Delete a user layout. This will also delete the file on disk. */
+	void DeleteUserLayout(const FString& LayoutName);
+
 	//~ Begin FWorkflowCentricApplication interface
 	virtual void SetCurrentMode(FName NewMode) override;
 	//~ End FWorkflowCentricApplication interface
 
 protected:
+	/** Get the path to the user settings. */
+	static FString UserSettingsDir()
+	{
+		static const FString UserDir = FPaths::Combine(FPlatformProcess::UserSettingsDir(), *FApp::GetEpicProductIdentifier(), TEXT("LiveLinkHub"));
+		return UserDir;
+	}
+
 	//~ Begin FWorkflowCentricApplication interface
 	virtual FName GetToolkitFName() const override { return "LiveLinkHub"; }
 	virtual FText GetBaseToolkitName() const override { return INVTEXT("LiveLinkHub Editor"); }
@@ -61,6 +93,12 @@ protected:
 	virtual FText GetToolkitName() const override;
 	//~ End AssetEditorToolkit interface
 
+	/** Iterate through layout directories to find .llhlayout files. */
+	void DiscoverLayouts();
+
+	/** Creates a user layout app mode from a layout file. */
+	bool RegisterUserLayout(const FString& LayoutName, const FString& LayoutFile);
+
 private:
 	/** Don't allow LiveLinkHub to call this directly since it wouldn't register the cached mode info. */
 	virtual void AddApplicationMode(FName ModeName, TSharedRef<FApplicationMode> Mode) override;
@@ -76,11 +114,29 @@ private:
 	/** Used to populate the status toolbar with widgets of the current mode. */
 	void AddToolbarExtenders(FToolBarBuilder&);
 
-private:
+	/** Persist layout and save it to the specified path. */
+	void SaveLayoutToFile(const FString& SavePath = TEXT(""), TSharedPtr<FJsonObject> JsonLayout = TSharedPtr<FJsonObject>{});
+
+	TArray<FString> GetLayoutDirectories() const;
+
+	/** Read the layout file and convert it to a JsonObject. */
+	TOptional<struct FLiveLinkHubUserLayout> ParseUserLayout(const FString& LayoutName);
+
+	/** Get an application mode from  */
+	TSharedPtr<FLiveLinkHubApplicationMode> FindApplicationMode(const FString& ModeName);
 
+	void RemoveLiveLinkHubApplicationMode(FName ModeName);
+
+private:
 	/** Delegate called when the active app mode has changed. */
 	FOnApplicationModeChanged AppModeChangedDelegate;
 
 	/** Map of mode name to cached mode info. */
 	TMap<FName, FLiveLinkHubAppModeInfo> CachedModeInfo;
-};
\ No newline at end of file
+
+	/** Path to the last saved layout. */
+	FString LastLayoutPath;
+
+	/** Layout name to file path */
+	TMap<FString, FString> CachedLayouts;
+};
diff --git a/Engine/Plugins/Animation/LiveLinkHub/Source/LiveLinkHub/Private/LiveLinkHubUserLayoutMode.cpp b/Engine/Plugins/Animation/LiveLinkHub/Source/LiveLinkHub/Private/LiveLinkHubUserLayoutMode.cpp
new file mode 100644
--- /dev/null
+++ b/Engine/Plugins/Animation/LiveLinkHub/Source/LiveLinkHub/Private/LiveLinkHubUserLayoutMode.cpp
@@ -0,0 +1,34 @@
+// Copyright Epic Games, Inc. All Rights Reserved.
+
+#include "LiveLinkHubUserLayoutMode.h"
+
+#include "Dom/JsonObject.h"
+#include "LiveLinkHub.h"
+
+FLiveLinkHubUserLayoutMode::FLiveLinkHubUserLayoutMode(FName InLayoutName, TSharedRef<FJsonObject> UserLayout, TSharedPtr<FLiveLinkHubApplicationMode> InParentMode)
+	: FLiveLinkHubApplicationMode(InLayoutName, FText::FromName(InLayoutName), FLiveLinkHub::Get())
+	, ParentMode(MoveTemp(InParentMode))
+{
+	LayoutExtender = ParentMode->LayoutExtender;
+	ToolbarExtender = ParentMode->GetToolbarExtender();
+
+	ParentLayoutName = UserLayout->GetStringField(TEXT("Name"));
+
+	// Note: At runtime, we *must* override the layout name so that it doesn't clash with the parent mode's layout name.
+	// Otherwise, we would run into issues where modifying this layout would also affect the parent or vice-versa.
+	UserLayout->SetStringField(TEXT("Name"), GetModeName().ToString());
+
+	TabLayout = FTabManager::FLayout::NewFromJson(UserLayout);
+}
+
+void FLiveLinkHubUserLayoutMode::PreDeactivateMode()
+{
+	FLiveLinkHubApplicationMode::PreDeactivateMode();
+
+	TSharedPtr<FTabManager::FLayout> PersistentLayout = FLiveLinkHub::Get()->GetTabManager()->PersistLayout();
+	TSharedPtr<FJsonObject> JsonLayout = PersistentLayout->ToJson();
+
+	// Save out the layout with the original name.
+	JsonLayout->SetStringField(TEXT("Name"), ParentLayoutName);
+	FLiveLinkHub::Get()->PersistUserLayout(ModeName.ToString(), JsonLayout);
+}
diff --git a/Engine/Plugins/Animation/LiveLinkHub/Source/LiveLinkHub/Private/LiveLinkHubUserLayoutMode.h b/Engine/Plugins/Animation/LiveLinkHub/Source/LiveLinkHub/Private/LiveLinkHubUserLayoutMode.h
new file mode 100644
--- /dev/null
+++ b/Engine/Plugins/Animation/LiveLinkHub/Source/LiveLinkHub/Private/LiveLinkHubUserLayoutMode.h
@@ -0,0 +1,50 @@
+// Copyright Epic Games, Inc. All Rights Reserved.
+
+#pragma once
+
+#include "LiveLinkHubApplicationMode.h"
+
+#include "Dom/JsonObject.h"
+
+/** ApplicationMode that represents a user layout. It acts as a proxy for its parent mode. */
+class FLiveLinkHubUserLayoutMode : public FLiveLinkHubApplicationMode
+{
+public:
+	FLiveLinkHubUserLayoutMode(FName InLayoutName, TSharedRef<FJsonObject> UserLayout, TSharedPtr<FLiveLinkHubApplicationMode> InParentMode);
+
+	/** Get the parent mode that holds the actual functionality (ie. Creator mode, Capture Manager) */
+	const TSharedPtr<FLiveLinkHubApplicationMode>& GetParentMode() const
+	{
+		return ParentMode;
+	}
+
+	//~ Begin FLiveLinkHubApplicationMode interface
+	virtual FSlateIcon GetModeIcon() const override 
+	{ 
+		return ParentMode->GetModeIcon();
+	}
+
+	virtual bool IsUserLayout() const override
+	{
+		return true;
+	}
+
+	virtual TArray<TSharedRef<SWidget>> GetStatusBarWidgets_Impl() override
+	{
+		return ParentMode->GetStatusBarWidgets_Impl();
+	}
+
+	virtual TArray<TSharedRef<SWidget>> GetToolbarWidgets_Impl() override
+	{
+		return ParentMode->GetToolbarWidgets_Impl();
+	}
+
+	virtual void PreDeactivateMode() override;
+	//~ End FLiveLinkHubApplicationMode interface
+
+private:
+	/** The mode that holds the functionality and tab spawners for this mode. */
+	TSharedPtr<FLiveLinkHubApplicationMode> ParentMode;
+	/** Layout name of the parent mode (ie. LiveLinkHubCreatorMode_V1.0) */
+	FString ParentLayoutName;
+};
\ No newline at end of file
diff --git a/Engine/Plugins/Animation/LiveLinkHub/Source/LiveLinkHub/Private/Settings/LiveLinkHubSettings.h b/Engine/Plugins/Animation/LiveLinkHub/Source/LiveLinkHub/Private/Settings/LiveLinkHubSettings.h
--- a/Engine/Plugins/Animation/LiveLinkHub/Source/LiveLinkHub/Private/Settings/LiveLinkHubSettings.h
+++ b/Engine/Plugins/Animation/LiveLinkHub/Source/LiveLinkHub/Private/Settings/LiveLinkHubSettings.h
@@ -107,3 +107,17 @@ private:
 	UPROPERTY(Instanced, Transient)
 	mutable TObjectPtr<ULiveLinkHubNamingTokens> NamingTokens;
 };
+
+/**
+ * User Settings for LiveLinkHub.
+ */
+UCLASS(config = EditorPerProjectUserSettings)
+class LIVELINKHUB_API ULiveLinkHubUserSettings : public UObject
+{
+	GENERATED_BODY()
+
+public:
+	/** Which directories to scan to discover layouts */
+	UPROPERTY(config, EditAnywhere, Category = "LiveLinkHub")
+	TArray<FString> LayoutDirectories;
+};
diff --git a/Engine/Plugins/Animation/LiveLinkHub/Source/LiveLinkHub/Private/UI/Widgets/SApplicationModeSwitcher.h b/Engine/Plugins/Animation/LiveLinkHub/Source/LiveLinkHub/Private/UI/Widgets/SApplicationModeSwitcher.h
--- a/Engine/Plugins/Animation/LiveLinkHub/Source/LiveLinkHub/Private/UI/Widgets/SApplicationModeSwitcher.h
+++ b/Engine/Plugins/Animation/LiveLinkHub/Source/LiveLinkHub/Private/UI/Widgets/SApplicationModeSwitcher.h
@@ -4,10 +4,17 @@
 
 #include "Widgets/SCompoundWidget.h"
 
+#include "DesktopPlatformModule.h"
+#include "EditorDirectories.h"
+#include "Features/IModularFeatures.h"
+#include "Features/IModularFeature.h"
+#include "LiveLinkHubClient.h"
 #include "LiveLinkHub.h"
+#include "Misc/FileHelper.h"
 #include "Widgets/Images/SImage.h"
 #include "Widgets/Input/SComboButton.h"
 
+#define LOCTEXT_NAMESPACE "ApplicationModeSwitcher"
 
 class SApplicationModeSwitcher : public SCompoundWidget
 {
@@ -62,30 +69,98 @@ public:
 	}
 
 private:
+	/** Create the content of the layout button when it's oppened. */
 	TSharedRef<SWidget> GetModeSwitcherContent()
 	{
 		FMenuBuilder MenuBuilder(true, NULL);
-		if (TSharedPtr<FLiveLinkHub> LiveLinkHub = FLiveLinkHub::Get())
+
+		TSharedPtr<FLiveLinkHub> LiveLinkHub = FLiveLinkHub::Get();
+		const TArray<FName> Modes = LiveLinkHub->GetApplicationModes();
+		const FName CurrentModeName = LiveLinkHub->GetCurrentMode();
+
+		MenuBuilder.BeginSection("Layouts", LOCTEXT("LayoutsLabel", "Layouts"));
+		
+		for (FName Mode : Modes)
 		{
-			TArray<FName> Modes = LiveLinkHub->GetApplicationModes();
-			FName CurrentModeName = LiveLinkHub->GetCurrentMode();
+			if (Mode != CurrentModeName)
+			{
+				if (TOptional<FLiveLinkHubAppModeInfo> AppModeInfo = LiveLinkHub->GetModeInfo(Mode))
+				{
+					// Tooltip text, Icon, localized name
+					MenuBuilder.AddMenuEntry(AppModeInfo->DisplayName, AppModeInfo->DisplayName, AppModeInfo->Icon, FExecuteAction::CreateRaw(this, &SApplicationModeSwitcher::SetActiveMode, Mode));
+				}
+			}
+		}
+
+		MenuBuilder.EndSection();
 
-			for (FName Mode : Modes)
+		MenuBuilder.BeginSection("UserLayouts", LOCTEXT("UserLayoutsLabel", "User Layouts"));
+
+		TArray<FString> UserLayouts = FLiveLinkHub::Get()->GetUserLayouts();
+		for (const FString& Layout : UserLayouts)
+		{
+			if (Layout != CurrentModeName)
 			{
-				if (Mode != CurrentModeName)
+				if (TOptional<FLiveLinkHubAppModeInfo> AppModeInfo = LiveLinkHub->GetModeInfo(*Layout))
 				{
-					if (TOptional<FLiveLinkHubAppModeInfo> AppModeInfo = LiveLinkHub->GetModeInfo(Mode))
-					{
-						// Tooltip text, Icon, localized name
-						MenuBuilder.AddMenuEntry(AppModeInfo->DisplayName, AppModeInfo->DisplayName, AppModeInfo->Icon, FExecuteAction::CreateRaw(this, &SApplicationModeSwitcher::SetActiveMode, Mode));
-					}
+					// Tooltip text, Icon, localized name
+					MenuBuilder.AddMenuEntry(AppModeInfo->DisplayName, AppModeInfo->DisplayName, AppModeInfo->Icon, FExecuteAction::CreateRaw(this, &SApplicationModeSwitcher::SetActiveMode, FName(Layout)));
 				}
 			}
 		}
 
+		MenuBuilder.EndSection();
+
+		MenuBuilder.AddSeparator();
+
+		MenuBuilder.AddMenuEntry(
+			INVTEXT("Save Layout As"),
+			TAttribute<FText>(),
+			FSlateIcon(),
+			FUIAction(FExecuteAction::CreateLambda([]() { FLiveLinkHub::Get()->SaveLayoutAs(); })),
+			NAME_None,
+			EUserInterfaceActionType::Button
+		);
+
+		MenuBuilder.AddMenuEntry(
+			INVTEXT("Load Layout"),
+			TAttribute<FText>(),
+			FSlateIcon(),
+			FUIAction(FExecuteAction::CreateLambda([]() { FLiveLinkHub::Get()->LoadLayout(); })),
+			NAME_None,
+			EUserInterfaceActionType::Button
+		);
+
+		if (TOptional<FLiveLinkHubAppModeInfo> AppModeInfo = LiveLinkHub->GetModeInfo(CurrentModeName))
+		{
+			if (!AppModeInfo->bUserLayout)
+			{
+				MenuBuilder.AddMenuEntry(
+					INVTEXT("Reset Current Layout"),
+					TAttribute<FText>(),
+					FSlateIcon(),
+					FUIAction(FExecuteAction::CreateLambda([]() { FLiveLinkHub::Get()->ResetLayout(); })),
+					NAME_None,
+					EUserInterfaceActionType::Button
+				);
+			}
+		}
+
+		if (UserLayouts.Num())
+		{
+			MenuBuilder.AddSubMenu(
+				INVTEXT("Delete Layout"),
+				TAttribute<FText>(),
+				FNewMenuDelegate::CreateSP(this, &SApplicationModeSwitcher::CreateDeleteLayoutMenu),
+				false,
+				FSlateIcon()
+			);
+		}
+
 		return MenuBuilder.MakeWidget();
 	}
 
+	/** Set the current application mode. */
 	void SetActiveMode(FName ModeName)
 	{
 		if (TSharedPtr<FLiveLinkHub> LiveLinkHub = FLiveLinkHub::Get())
@@ -94,6 +169,7 @@ private:
 		}
 	}
 
+	/** Handle app mode changing. */
 	void OnAppModeChanged(FName NewModeName) const
 	{
 		if (TSharedPtr<FLiveLinkHub> LiveLinkHub = FLiveLinkHub::Get())
@@ -106,9 +182,34 @@ private:
 		}
 	}
 
-private:
+	/** Populate the Delete Layout menu with the list of user layouts. */
+	void CreateDeleteLayoutMenu(FMenuBuilder& MenuBuilder)
+	{
+		TSharedPtr<FLiveLinkHub> LiveLinkHub = FLiveLinkHub::Get();
+		const FName CurrentModeName = LiveLinkHub->GetCurrentMode();
 
+		TArray<FString> UserLayouts = FLiveLinkHub::Get()->GetUserLayouts();
+		for (const FString& Layout : UserLayouts)
+		{
+			if (CurrentModeName != Layout)
+			{
+				MenuBuilder.AddMenuEntry(
+					FText::FromString(Layout),
+					TAttribute<FText>(),
+					FSlateIcon(),
+					FUIAction(FExecuteAction::CreateLambda([Layout]() { FLiveLinkHub::Get()->DeleteUserLayout(Layout); })),
+					NAME_None,
+					EUserInterfaceActionType::Button
+				);
+			}
+		}
+	}
+
+private:
+	/** Holds the icon of the current mode. */
 	TSharedPtr<SImage> ActiveModeImage;
+	/** Holds the display name of the current mode. */
 	TSharedPtr<STextBlock> ActiveModeDisplayName;
 };
 
+#undef LOCTEXT_NAMESPACE
\ No newline at end of file
diff --git a/Engine/Plugins/Animation/LiveLinkHub/Source/LiveLinkHub/Public/LiveLinkHubApplicationMode.h b/Engine/Plugins/Animation/LiveLinkHub/Source/LiveLinkHub/Public/LiveLinkHubApplicationMode.h
--- a/Engine/Plugins/Animation/LiveLinkHub/Source/LiveLinkHub/Public/LiveLinkHubApplicationMode.h
+++ b/Engine/Plugins/Animation/LiveLinkHub/Source/LiveLinkHub/Public/LiveLinkHubApplicationMode.h
@@ -44,6 +44,37 @@ public:
 	/** Method used to gather the toolbar widgets that should be visible when this mode is active. */
 	TArray<TSharedRef<SWidget>> GetToolbarWidgets();
 
+	/** Get the name of this layout. */
+	FName GetLayoutName()
+	{
+		return TabLayout->GetLayoutName();
+	}
+
+	/** Get the filename of this mode's layout ini. */
+	FString GetLayoutIni() const
+	{
+		return LayoutIni;
+	}
+
+	/** Get the tab manager layout ptr. */
+	TSharedPtr<FTabManager::FLayout> GetTabLayout() const
+	{
+		return TabLayout;
+	}
+
+	/** Returns whether this mode is a user layout. */
+	virtual bool IsUserLayout() const
+	{
+		return false;
+	}
+
+	/** Implement in a child class in order to have status bar extensions show up when the mode is active. */
+	virtual TArray<TSharedRef<SWidget>> GetToolbarWidgets_Impl() { return {}; };
+
+
+	/** Get widgets that should appear in the status bar. */
+	virtual TArray<TSharedRef<SWidget>> GetStatusBarWidgets_Impl() { return {}; }
+
 protected:
 	FLiveLinkHubApplicationMode(FName InApplicationMode, FText InDisplayName, TSharedPtr<FLiveLinkHubApplicationBase> InApp);
 
@@ -64,14 +95,7 @@ protected:
 			bRegisteredTabFactories = true;
 		}
 	}
-	//~ End FApplicationMode interface
-
-	/** Implement in a child class in order to have status bar extensions show up when the mode is active. */
-	virtual TArray<TSharedRef<SWidget>> GetToolbarWidgets_Impl() { return {}; };
 
-	
-	/** Get widgets that should appear in the status bar. */
-	virtual TArray<TSharedRef<SWidget>> GetStatusBarWidgets_Impl() { return {}; }
 protected:
 	/** The application that hosts this app mode. */
 	TWeakPtr<FLiveLinkHubApplicationBase> WeakApp;
diff --git a/Engine/Plugins/VirtualProduction/CaptureManager/CaptureManagerApp/Source/CaptureManagerEditor/Private/UI/LiveLinkHubCaptureManagerMode.cpp b/Engine/Plugins/VirtualProduction/CaptureManager/CaptureManagerApp/Source/CaptureManagerEditor/Private/UI/LiveLinkHubCaptureManagerMode.cpp
--- a/Engine/Plugins/VirtualProduction/CaptureManager/CaptureManagerApp/Source/CaptureManagerEditor/Private/UI/LiveLinkHubCaptureManagerMode.cpp
+++ b/Engine/Plugins/VirtualProduction/CaptureManager/CaptureManagerApp/Source/CaptureManagerEditor/Private/UI/LiveLinkHubCaptureManagerMode.cpp
@@ -131,7 +131,7 @@ FLiveLinkHubCaptureManagerMode::FLiveLinkHubCaptureManagerMode(TSharedPtr<FLiveL
 {
 	using namespace UE::CaptureManager;
 
-	TabLayout = FTabManager::NewLayout("LiveLinkCaptureManagerMode_v1.0");
+	TabLayout = FTabManager::NewLayout("LiveLinkCaptureManagerMode_v1.1");
 	
 	TOptional<FString> LocalIpAddressOpt = GetLocalIpAddress();
 
