Live Link Hub - Fix tab tool tip to show "Live Link Hub" instead of "Asset: Default_Object"

diff --git a/Engine/Plugins/Animation/LiveLinkHub/Source/LiveLinkHub/Private/LiveLinkHubApplication.cpp b/Engine/Plugins/Animation/LiveLinkHub/Source/LiveLinkHub/Private/LiveLinkHubApplication.cpp
--- a/Engine/Plugins/Animation/LiveLinkHub/Source/LiveLinkHub/Private/LiveLinkHubApplication.cpp
+++ b/Engine/Plugins/Animation/LiveLinkHub/Source/LiveLinkHub/Private/LiveLinkHubApplication.cpp
@@ -566,6 +566,11 @@ FText FLiveLinkHubApplication::GetToolkitName() const
 	return LOCTEXT("LiveLinkHubLabel", "Live Link Hub");
 }
 
+FText FLiveLinkHubApplication::GetToolkitToolTipText() const
+{
+	return GetToolkitName();
+}
+
 void FLiveLinkHubApplication::DiscoverLayouts()
 {
 	TArray<FString> Files;
diff --git a/Engine/Plugins/Animation/LiveLinkHub/Source/LiveLinkHub/Private/LiveLinkHubApplication.h b/Engine/Plugins/Animation/LiveLinkHub/Source/LiveLinkHub/Private/LiveLinkHubApplication.h
--- a/Engine/Plugins/Animation/LiveLinkHub/Source/LiveLinkHub/Private/LiveLinkHubApplication.h
+++ b/Engine/Plugins/Animation/LiveLinkHub/Source/LiveLinkHub/Private/LiveLinkHubApplication.h
@@ -91,6 +91,7 @@ protected:
 	virtual TSharedPtr<SWidget> CreateMenuBar(const TSharedPtr<FTabManager>& InTabManager, const FName InMenuName, FToolMenuContext& InToolMenuContext) override;
 	virtual void MapToolkitCommands() {}
 	virtual FText GetToolkitName() const override;
+	virtual FText GetToolkitToolTipText() const override;
 	//~ End AssetEditorToolkit interface
 
 	/** Iterate through layout directories to find .llhlayout files. */
