Live Link Hub - Rename playback settings to better indicate that they're playback settings

#rnx

diff --git a/Engine/Plugins/Animation/LiveLinkHub/Source/LiveLinkHub/Private/Recording/Implementations/LiveLinkFrameData.cpp b/Engine/Plugins/Animation/LiveLinkHub/Source/LiveLinkHub/Private/Recording/Implementations/LiveLinkFrameData.cpp
--- a/Engine/Plugins/Animation/LiveLinkHub/Source/LiveLinkHub/Private/Recording/Implementations/LiveLinkFrameData.cpp
+++ b/Engine/Plugins/Animation/LiveLinkHub/Source/LiveLinkHub/Private/Recording/Implementations/LiveLinkFrameData.cpp
@@ -8,7 +8,7 @@ namespace UE::LiveLinkHub::FrameData::Private
 {
 	void FFrameBufferCache::TrimCache()
 	{
-		const int32 MaxHistory = GetDefault<ULiveLinkHubSettings>()->MaxBufferRangeHistory;
+		const int32 MaxHistory = GetDefault<ULiveLinkHubSettings>()->PlaybackMaxBufferRangeHistory;
 		const int32 AmountToTrim = FrameData.Num() - MaxHistory;
 		if (AmountToTrim > 0)
 		{
diff --git a/Engine/Plugins/Animation/LiveLinkHub/Source/LiveLinkHub/Private/Recording/Implementations/LiveLinkUAssetRecording.cpp b/Engine/Plugins/Animation/LiveLinkHub/Source/LiveLinkHub/Private/Recording/Implementations/LiveLinkUAssetRecording.cpp
--- a/Engine/Plugins/Animation/LiveLinkHub/Source/LiveLinkHub/Private/Recording/Implementations/LiveLinkUAssetRecording.cpp
+++ b/Engine/Plugins/Animation/LiveLinkHub/Source/LiveLinkHub/Private/Recording/Implementations/LiveLinkUAssetRecording.cpp
@@ -896,7 +896,7 @@ void ULiveLinkUAssetRecording::LoadFrameData(UE::LiveLinkHub::FrameData::Private
 			bLoadRight = bRightOnly ? true : !bLoadRight;
 		};
 
-		const int32 BufferBatchSize = bForceSequential ? MaxFrames : GetDefault<ULiveLinkHubSettings>()->BufferBatchSize;
+		const int32 BufferBatchSize = bForceSequential ? MaxFrames : GetDefault<ULiveLinkHubSettings>()->PlaybackBufferBatchSize;
 
 		struct FBulkDataBatch
 		{
diff --git a/Engine/Plugins/Animation/LiveLinkHub/Source/LiveLinkHub/Private/Recording/Implementations/LiveLinkUAssetRecordingPlayer.cpp b/Engine/Plugins/Animation/LiveLinkHub/Source/LiveLinkHub/Private/Recording/Implementations/LiveLinkUAssetRecordingPlayer.cpp
--- a/Engine/Plugins/Animation/LiveLinkHub/Source/LiveLinkHub/Private/Recording/Implementations/LiveLinkUAssetRecordingPlayer.cpp
+++ b/Engine/Plugins/Animation/LiveLinkHub/Source/LiveLinkHub/Private/Recording/Implementations/LiveLinkUAssetRecordingPlayer.cpp
@@ -387,7 +387,7 @@ int32 FLiveLinkUAssetRecordingPlayer::GetNumFramesToBuffer() const
 	if (LoadedRecording.IsValid())
 	{
 		const int32 FrameSize = LoadedRecording->GetFrameDiskSize();
-		const int32 MaxFrameBufferSizeMB = GetDefault<ULiveLinkHubSettings>()->FrameBufferSizeMB;
+		const int32 MaxFrameBufferSizeMB = GetDefault<ULiveLinkHubSettings>()->PlaybackFrameBufferSizeMB;
 		const int32 MaxFrameBufferSizeBytes = MaxFrameBufferSizeMB * 1024 * 1024;
 
 		// We divide total frames by 2, since they get doubled later to account for scrubbing in both directions.
diff --git a/Engine/Plugins/Animation/LiveLinkHub/Source/LiveLinkHub/Private/Settings/LiveLinkHubSettings.h b/Engine/Plugins/Animation/LiveLinkHub/Source/LiveLinkHub/Private/Settings/LiveLinkHubSettings.h
--- a/Engine/Plugins/Animation/LiveLinkHub/Source/LiveLinkHub/Private/Settings/LiveLinkHubSettings.h
+++ b/Engine/Plugins/Animation/LiveLinkHub/Source/LiveLinkHub/Private/Settings/LiveLinkHubSettings.h
@@ -41,15 +41,15 @@ public:
 
 	/** The size in megabytes to buffer when streaming a recording. */
 	UPROPERTY(config, EditAnywhere, Category="LiveLinkHub", meta = (ClampMin = "1", UIMin = "1"))
-	int32 FrameBufferSizeMB = 100;
+	int32 PlaybackFrameBufferSizeMB = 100;
 
 	/** Number of frames to buffer at once. */
 	UPROPERTY(config, EditAnywhere, AdvancedDisplay, Category="LiveLinkHub", meta = (ClampMin = "2", UIMin = "2"))
-	int32 BufferBatchSize = 5;
+	int32 PlaybackBufferBatchSize = 5;
 
 	/** Maximum number of frame ranges to store in history while scrubbing. Increasing can make scrubbing faster but temporarily use more memory. */
 	UPROPERTY(config, EditAnywhere, AdvancedDisplay, Category="LiveLinkHub")
-	int32 MaxBufferRangeHistory = 25;
+	int32 PlaybackMaxBufferRangeHistory = 25;
 	
 	/** Which project settings sections to display when opening the settings viewer. */
 	UPROPERTY(config)
@@ -83,7 +83,7 @@ public:
 	UPROPERTY(config, EditAnywhere, Category="LiveLinkHub", meta = (ConfigRestartRequired = true, ClampMin="15.0"))
 	float TargetFrameRate = 60.0f;
 
-	/** Maximum time in seconds to wait for sources to clean up. */
+	/** Maximum time in seconds to wait for sources to clean up. Increase this value if you notice that some sources are incorrectly cleaned up when switching a config. */
 	UPROPERTY(config, EditAnywhere, AdvancedDisplay, Category="LiveLinkHub", meta = (ClampMin="0.0"))
 	float SourceMaxCleanupTime = 0.25f;
 	
