LiveLinkHub
- Better default on the FileNameTemplate setting

diff --git a/Engine/Plugins/Animation/LiveLinkHub/Source/LiveLinkHub/Private/Settings/LiveLinkHubSettings.h b/Engine/Plugins/Animation/LiveLinkHub/Source/LiveLinkHub/Private/Settings/LiveLinkHubSettings.h
--- a/Engine/Plugins/Animation/LiveLinkHub/Source/LiveLinkHub/Private/Settings/LiveLinkHubSettings.h
+++ b/Engine/Plugins/Animation/LiveLinkHub/Source/LiveLinkHub/Private/Settings/LiveLinkHubSettings.h
@@ -89,7 +89,7 @@ public:
 	
 	/** The filename template to use when creating recordings. */
 	UPROPERTY(config, EditAnywhere, Category="Templates")
-	FString FilenameTemplate = TEXT("NewLiveLinkRecording");
+	FString FilenameTemplate = TEXT("{session}_{slate}_tk{take}");
 
 	/** Example parsed output of the template. */
 	UPROPERTY(VisibleAnywhere, Category="Templates", DisplayName="Output")
