LiveLinkHub - Move playback widget to a tab sitting above the status bar
Add devices tab to capture manager

diff --git a/Engine/Plugins/Animation/LiveLink/Source/LiveLinkEditor/Private/LiveLinkEditorModule.cpp b/Engine/Plugins/Animation/LiveLink/Source/LiveLinkEditor/Private/LiveLinkEditorModule.cpp
--- a/Engine/Plugins/Animation/LiveLink/Source/LiveLinkEditor/Private/LiveLinkEditorModule.cpp
+++ b/Engine/Plugins/Animation/LiveLink/Source/LiveLinkEditor/Private/LiveLinkEditorModule.cpp
@@ -158,6 +158,15 @@ public:
 
 		StyleSet->Set("EditableTextBox", EditableTextStyle);
 
+		FButtonStyle EjectButtonStyle = FAppStyle::Get().GetWidgetStyle<FButtonStyle>("Animation.PlayControlsButton");
+
+		const FCheckBoxStyle& CheckBoxStyle = FAppStyle::Get().GetWidgetStyle<FCheckBoxStyle>("ToggleButtonCheckbox");
+		EjectButtonStyle.Normal = CheckBoxStyle.CheckedImage;
+		EjectButtonStyle.Hovered = CheckBoxStyle.CheckedHoveredImage;
+
+		StyleSet->Set("LiveLinkHub.EjectButton", EjectButtonStyle);
+
+
 		FSlateStyleRegistry::RegisterSlateStyle(*StyleSet.Get());
 
 		RegisterSettings();
diff --git a/Engine/Plugins/Animation/LiveLinkHub/Source/LiveLinkHub/Private/LiveLinkHubCreatorAppMode.h b/Engine/Plugins/Animation/LiveLinkHub/Source/LiveLinkHub/Private/LiveLinkHubCreatorAppMode.h
--- a/Engine/Plugins/Animation/LiveLinkHub/Source/LiveLinkHub/Private/LiveLinkHubCreatorAppMode.h
+++ b/Engine/Plugins/Animation/LiveLinkHub/Source/LiveLinkHub/Private/LiveLinkHubCreatorAppMode.h
@@ -37,18 +37,6 @@
 
 #define LOCTEXT_NAMESPACE "LiveLinkHubCreatorMode"
 
-namespace UE::LiveLinkHub
-{
-	static const FName LiveLinkStyleName = "LiveLinkStyle";
-
-	static const FName SourcesTabId = "SourcesTabId";
-	static const FName SourceDetailsTabId = "SourceDetailsTabId";
-	static const FName SubjectsTabId = "SubjectsTabId";
-	static const FName PlaybackTabId = "PlaybackTabId";
-	static const FName ClientsTabId = "ClientsTabId";
-	static const FName ClientDetailsTabId = "ClientDetailsTabId";
-}
-
 struct FLiveLinkHubSourcesTabSummoner : public FWorkflowTabFactory
 {
 public:
@@ -208,6 +196,30 @@ private:
 	TWeakPtr<FLiveLinkPanelController> WeakPanelController;
 };
 
+
+struct FLiveLinkHubRecordingListTabSummoner : public FWorkflowTabFactory
+{
+public:
+	FLiveLinkHubRecordingListTabSummoner(TSharedPtr<class FAssetEditorToolkit> InHostingApp)
+		: FWorkflowTabFactory(UE::LiveLinkHub::RecordingListTabId, MoveTemp(InHostingApp))
+	{
+		TabLabel = LOCTEXT("RecordingListTabLabel", "Recordings List");
+		TabIcon = FSlateIcon(UE::LiveLinkHub::LiveLinkStyleName, TEXT("LiveLinkHub.Playback.Icon"));
+
+		bIsSingleton = true;
+	}
+
+	virtual TSharedRef<SWidget> CreateTabBody(const FWorkflowTabSpawnInfo& Info) const override
+	{
+		return FLiveLinkHub::Get()->GetRecordingListController()->MakeRecordingList();
+	}
+
+	virtual FText GetTabToolTipText(const FWorkflowTabSpawnInfo& Info) const override
+	{
+		return LOCTEXT("RecordingListTabToolTip", "Hosts the LiveLink recordings list.");
+	}
+};
+
 struct FLiveLinkHubPlaybackTabSummoner : public FWorkflowTabFactory
 {
 public:
@@ -216,25 +228,14 @@ public:
 	{
 		TabLabel = LOCTEXT("PlaybackTabLabel", "Playback");
 		TabIcon = FSlateIcon(UE::LiveLinkHub::LiveLinkStyleName, TEXT("LiveLinkHub.Playback.Icon"));
+		bShouldAutosize = true;
 
 		bIsSingleton = true;
 	}
 
 	virtual TSharedRef<SWidget> CreateTabBody(const FWorkflowTabSpawnInfo& Info) const override
 	{
-		FLiveLinkHubModule& LiveLinkHubModule = FModuleManager::Get().GetModuleChecked<FLiveLinkHubModule>("LiveLinkHub");
-
-		return SNew(SVerticalBox)
-			+ SVerticalBox::Slot()
-			.FillHeight(1.0f)
-			[
-				LiveLinkHubModule.GetRecordingListController()->MakeRecordingList()
-			]
-			+ SVerticalBox::Slot()
-			.AutoHeight()
-			[
-				LiveLinkHubModule.GetPlaybackController()->MakePlaybackWidget()
-			];
+		return FLiveLinkHub::Get()->GetPlaybackController()->MakePlaybackWidget();
 	}
 
 	virtual FText GetTabToolTipText(const FWorkflowTabSpawnInfo& Info) const override
@@ -303,7 +304,7 @@ public:
 		: FLiveLinkHubApplicationMode("CreatorMode", LOCTEXT("CreatorModeLabel", "Creator Mode"), App)
 	{
 
-		TabLayout = FTabManager::NewLayout("LiveLinkHubCreatorMode_v1.0");
+		TabLayout = FTabManager::NewLayout("LiveLinkHubCreatorMode_v1.1");
 		const TSharedRef<FTabManager::FArea> MainWindowArea = FTabManager::NewPrimaryArea()->SetOrientation(Orient_Horizontal);
 
 		PanelController = MakeShared<FLiveLinkPanelController>(TAttribute<bool>::CreateRaw(this, &FLiveLinkHubCreatorAppMode::IsSourcePanelReadOnly));
@@ -347,7 +348,7 @@ public:
 						FTabManager::NewStack()
 						->SetSizeCoefficient(0.5f)
 						->AddTab(UE::LiveLinkHub::SubjectsTabId, ETabState::OpenedTab)
-						->AddTab(UE::LiveLinkHub::PlaybackTabId, ETabState::OpenedTab)
+						->AddTab(UE::LiveLinkHub::RecordingListTabId, ETabState::OpenedTab)
 						->SetForegroundTab(UE::LiveLinkHub::SubjectsTabId)
 					)
 				)
@@ -372,6 +373,12 @@ public:
 					)
 				)
 			)
+			->Split
+			(
+				FTabManager::NewStack()
+				->SetHideTabWell(true)
+				->AddTab(UE::LiveLinkHub::PlaybackTabId, ETabState::ClosedTab)
+			)
 		);
 		
 		ToolbarExtender->AddToolBarExtension("Asset", EExtensionHook::After, App->GetToolkitCommands(), FToolBarExtensionDelegate::CreateRaw(this, &FLiveLinkHubCreatorAppMode::AddRecordingInfoToolbarExtension));
@@ -380,6 +387,7 @@ public:
 		TabFactories.RegisterFactory(MakeShared<FLiveLinkHubSourceDetailsTabSummoner>(App, PanelController));
 		TabFactories.RegisterFactory(MakeShared<FLiveLinkHubSubjectsTabSummoner>(App, PanelController));
 		TabFactories.RegisterFactory(MakeShared<FLiveLinkHubPlaybackTabSummoner>(App));
+		TabFactories.RegisterFactory(MakeShared<FLiveLinkHubRecordingListTabSummoner>(App));
 		TabFactories.RegisterFactory(MakeShared<FLiveLinkHubClientsTabSummoner>(App));
 		TabFactories.RegisterFactory(MakeShared<FLiveLinkHubClientDetailsTabSummoner>(App));
 	}
@@ -413,22 +421,8 @@ public:
 	virtual TArray<TSharedRef<SWidget>> GetStatusBarWidgets_Impl() override
 	{
 		return {
-			SNew(SHorizontalBox)
-			+ SHorizontalBox::Slot()
-			.AutoWidth()
-			.VAlign(VAlign_Center)
-			.Padding(FMargin(2.0, 0.0, 4.0, 0.0))
-			[
-				FLiveLinkHub::Get()->GetPlaybackController()->MakePlaybackControlWidget()
-			]
-			+ SHorizontalBox::Slot()
-			.AutoWidth()
-			.VAlign(VAlign_Center)
-			.Padding(2.0, 0.0)
-			[
-				SNew(STextBlock)
-				.Text_Raw(this, &FLiveLinkHubCreatorAppMode::GetLoadedConfigText)
-			]
+			SNew(STextBlock)
+			.Text_Raw(this, &FLiveLinkHubCreatorAppMode::GetLoadedConfigText)
 		};
 	}
 	//~ End FLiveLinkHubApplicationMode interface
diff --git a/Engine/Plugins/Animation/LiveLinkHub/Source/LiveLinkHub/Private/Recording/LiveLinkHubPlaybackController.cpp b/Engine/Plugins/Animation/LiveLinkHub/Source/LiveLinkHub/Private/Recording/LiveLinkHubPlaybackController.cpp
--- a/Engine/Plugins/Animation/LiveLinkHub/Source/LiveLinkHub/Private/Recording/LiveLinkHubPlaybackController.cpp
+++ b/Engine/Plugins/Animation/LiveLinkHub/Source/LiveLinkHub/Private/Recording/LiveLinkHubPlaybackController.cpp
@@ -7,6 +7,7 @@
 #include "CoreGlobals.h"
 #include "ILiveLinkClient.h"
 #include "Implementations/LiveLinkUAssetRecordingPlayer.h"
+#include "LiveLinkHubCreatorAppMode.h"
 #include "LiveLinkHubLog.h"
 #include "LiveLinkRecording.h"
 #include "LiveLinkPreset.h"
@@ -19,7 +20,6 @@
 #include "LiveLinkHubModule.h"
 #include "Modules/ModuleManager.h"
 #include "PackageTools.h"
-#include "Recording/SLiveLinkHubPlaybackControl.h"
 #include "UI/Widgets/SLiveLinkHubPlaybackWidget.h"
 #include "UObject/Object.h"
 #include "UObject/Package.h"
@@ -95,7 +95,8 @@ FLiveLinkHubPlaybackController::~FLiveLinkHubPlaybackController()
 TSharedRef<SWidget> FLiveLinkHubPlaybackController::MakePlaybackWidget()
 {
 	return SNew(SLiveLinkHubPlaybackWidget)
-		.IsEnabled_Raw(this, &FLiveLinkHubPlaybackController::IsReady)
+		.Visibility_Lambda([this]() { return IsInPlayback() ? EVisibility::Visible : EVisibility::Collapsed; })
+		.IsPlaybackEnabled_Raw(this, &FLiveLinkHubPlaybackController::IsReady)
 		.OnPlayForward_Raw(this, &FLiveLinkHubPlaybackController::BeginPlayback, false)
 		.OnPlayReverse_Raw(this, &FLiveLinkHubPlaybackController::BeginPlayback, true)
 		.OnFirstFrame_Lambda([this]()
@@ -114,6 +115,10 @@ TSharedRef<SWidget> FLiveLinkHubPlaybackController::MakePlaybackWidget()
 		{
 			GoToTime(FQualifiedFrameTime(FFrameTime(GetCurrentFrame() + 1), GetFrameRate()));
 		})
+		.OnExitPlayback_Lambda([this]()
+		{
+			EjectAndUnload();
+		})
 		.SetCurrentTime_Raw(this, &FLiveLinkHubPlaybackController::GoToTime)
 		.GetViewRange_Lambda([this]()
 		{
@@ -137,17 +142,6 @@ TSharedRef<SWidget> FLiveLinkHubPlaybackController::MakePlaybackWidget()
 		.GetFrameRate_Raw(this, &FLiveLinkHubPlaybackController::GetFrameRate);
 }
 
-TSharedRef<SWidget> FLiveLinkHubPlaybackController::MakePlaybackControlWidget()
-{
-	return SNew(SLiveLinkHubPlaybackControl)
-		.OnGetPlaybackState_Raw(this, &FLiveLinkHubPlaybackController::OnGetPlaybackState)
-		.OnResumePlayback_Raw(this, &FLiveLinkHubPlaybackController::OnResumePlayback)
-		.OnPausePlayback_Raw(this, &FLiveLinkHubPlaybackController::OnPausePlayback)
-		.OnExitPlayback_Raw(this, &FLiveLinkHubPlaybackController::OnExitPlayback)
-		.OnGetPlaybackTime_Raw(this, &FLiveLinkHubPlaybackController::OnGetPlaybackTime);
-}
-
-
 void FLiveLinkHubPlaybackController::StartPlayback()
 {
 	ResumePlayback();
@@ -197,6 +191,8 @@ void FLiveLinkHubPlaybackController::PreparePlayback(ULiveLinkRecording* InLiveL
 				// Make sure PreparingPlayback is set to false when the scope exits.
 				TGuardValue<bool> PreparingPlaybackGuard(bIsPreparingPlayback, true);
 
+				FLiveLinkHub::Get()->GetTabManager()->TryInvokeTab(UE::LiveLinkHub::PlaybackTabId);
+
 				RecordingToPlay.Reset(RecordingStrongPtr.Get());
 				if (!RecordingPlayer->PreparePlayback(RecordingToPlay.Get()))
 				{
@@ -243,7 +239,8 @@ void FLiveLinkHubPlaybackController::PreparePlayback(ULiveLinkRecording* InLiveL
 
 		if (RecordingToPlay.IsValid())
 		{
-			EjectAndUnload(PreparePlaybackCallback);
+			constexpr bool bClosePlaybackTab = false;
+			EjectAndUnload(PreparePlaybackCallback, nullptr, bClosePlaybackTab);
 		}
 		else
 		{
@@ -369,12 +366,20 @@ void FLiveLinkHubPlaybackController::Eject(TFunction<void()> CompletionCallback)
 	}
 }
 
-void FLiveLinkHubPlaybackController::EjectAndUnload(TFunction<void()> EjectCompletionCallback, const ULiveLinkRecording* InRecording)
+void FLiveLinkHubPlaybackController::EjectAndUnload(TFunction<void()> EjectCompletionCallback, const ULiveLinkRecording* InRecording, bool bClosePlaybackTab)
 {
 	const ULiveLinkRecording* Recording = InRecording ? InRecording : RecordingToPlay.Get();
 	UPackage* Package = Recording ? Recording->GetPackage() : nullptr;
 	const bool bIsSavingRecordingData = Recording && Recording->IsSavingRecordingData();
 
+	if (bClosePlaybackTab)
+	{
+		if (TSharedPtr<SDockTab> PlaybackTab = FLiveLinkHub::Get()->GetTabManager()->FindExistingLiveTab(UE::LiveLinkHub::PlaybackTabId))
+		{
+			PlaybackTab->RequestCloseTab();
+		}
+	}
+
 	// Only eject if this is in reference to the current recording playing.
 	if (InRecording == nullptr || InRecording == RecordingToPlay.Get())
 	{
@@ -663,25 +668,6 @@ bool FLiveLinkHubPlaybackController::ShouldRestart() const
 		|| (!bIsReverse && CurrentFrame >= GetSelectionEndTime().Time.GetFrame()));
 }
 
-ELiveLinkHubPlaybackState FLiveLinkHubPlaybackController::OnGetPlaybackState() const
-{
-	ELiveLinkHubPlaybackState PlaybackState = ELiveLinkHubPlaybackState::NotInPlayback;
-
-	if (IsInPlayback())
-	{
-		if (IsPlaying() && !IsPaused())
-		{
-			PlaybackState = ELiveLinkHubPlaybackState::Playing;
-		}
-		else if (IsPaused())
-		{
-			PlaybackState = ELiveLinkHubPlaybackState::Paused;
-		}
-	}
-
-	return PlaybackState;
-}
-
 void FLiveLinkHubPlaybackController::OnResumePlayback()
 {
 	if (ensureAlways(IsInPlayback() && IsPaused()))
@@ -698,13 +684,3 @@ void FLiveLinkHubPlaybackController::OnPausePlayback()
 		PausePlayback();
 	}
 }
-
-void FLiveLinkHubPlaybackController::OnExitPlayback()
-{
-	EjectAndUnload();
-}
-
-FQualifiedFrameTime FLiveLinkHubPlaybackController::OnGetPlaybackTime() const
-{
-	return GetCurrentTime();
-}
\ No newline at end of file
diff --git a/Engine/Plugins/Animation/LiveLinkHub/Source/LiveLinkHub/Private/Recording/LiveLinkHubPlaybackController.h b/Engine/Plugins/Animation/LiveLinkHub/Source/LiveLinkHub/Private/Recording/LiveLinkHubPlaybackController.h
--- a/Engine/Plugins/Animation/LiveLinkHub/Source/LiveLinkHub/Private/Recording/LiveLinkHubPlaybackController.h
+++ b/Engine/Plugins/Animation/LiveLinkHub/Source/LiveLinkHub/Private/Recording/LiveLinkHubPlaybackController.h
@@ -35,9 +35,6 @@ public:
 	/** Create the playback widget. */
 	TSharedRef<SWidget> MakePlaybackWidget();
 
-	/** Create a lightweight playback widget. */
-	TSharedRef<SWidget> MakePlaybackControlWidget();
-
 	/** Apply the recording's preset then prepare the data needed to playback. */
 	void PreparePlayback(ULiveLinkRecording* InLiveLinkRecording);
 	
@@ -67,7 +64,7 @@ public:
 	 * @param EjectCompletionCallback Callback after ejection occurs.
 	 * @param InRecording The specific recording to eject and unload. If null, the RecordingToPlay is used. 
 	 */
-	void EjectAndUnload(TFunction<void()> EjectCompletionCallback = nullptr, const ULiveLinkRecording* InRecording = nullptr);
+	void EjectAndUnload(TFunction<void()> EjectCompletionCallback = nullptr, const ULiveLinkRecording* InRecording = nullptr, bool bClosePlaybackTab = true);
 
 	/**
 	 * Completely unload a recording package.
@@ -201,21 +198,12 @@ private:
 	/** Checks if the current playback settings indicates the recording should restart. */
 	bool ShouldRestart() const;
 
-	/** Get the state of the playback. */
-	ELiveLinkHubPlaybackState OnGetPlaybackState() const;
-
 	/** Handles the user clicking the play button to resume playback. */
 	void OnResumePlayback();
 
 	/** Handles the user clicking the pause button to resume playback. */
 	void OnPausePlayback();
 
-	/** Handles the user clicking the eject button to exit playback. */
-	void OnExitPlayback();
-
-	/** Returns the current time in the playback. */
-	FQualifiedFrameTime OnGetPlaybackTime() const;
-
 private:
 	/** If the system has established a connection with the client. */
 	bool bIsReady = false;
diff --git a/Engine/Plugins/Animation/LiveLinkHub/Source/LiveLinkHub/Private/Recording/SLiveLinkHubPlaybackControl.h b/Engine/Plugins/Animation/LiveLinkHub/Source/LiveLinkHub/Private/Recording/SLiveLinkHubPlaybackControl.h
deleted file mode 100644
diff --git a/Engine/Plugins/Animation/LiveLinkHub/Source/LiveLinkHub/Private/UI/Widgets/SLiveLinkHubPlaybackWidget.cpp b/Engine/Plugins/Animation/LiveLinkHub/Source/LiveLinkHub/Private/UI/Widgets/SLiveLinkHubPlaybackWidget.cpp
--- a/Engine/Plugins/Animation/LiveLinkHub/Source/LiveLinkHub/Private/UI/Widgets/SLiveLinkHubPlaybackWidget.cpp
+++ b/Engine/Plugins/Animation/LiveLinkHub/Source/LiveLinkHub/Private/UI/Widgets/SLiveLinkHubPlaybackWidget.cpp
@@ -8,10 +8,12 @@
 #include "FrameNumberNumericInterface.h"
 #include "Framework/MultiBox/MultiBoxBuilder.h"
 #include "SSimpleTimeSlider.h"
+#include "Styling/SlateStyleRegistry.h"
 #include "Widgets/Input/SButton.h"
 #include "Widgets/Input/SComboButton.h"
 #include "Widgets/Input/SSpinBox.h"
 
+
 #define LOCTEXT_NAMESPACE "SLiveLinkHubPlaybackWidget"
 
 void SLiveLinkHubPlaybackWidget::Construct(const FArguments& InArgs)
@@ -25,6 +27,8 @@ void SLiveLinkHubPlaybackWidget::Construct(const FArguments& InArgs)
 	OnPreviousFrameDelegate = InArgs._OnPreviousFrame;
 	OnNextFrameDelegate = InArgs._OnNextFrame;
 
+	OnExitPlaybackDelegate = InArgs._OnExitPlayback;
+
 	OnGetPausedDelegate = InArgs._IsPaused;
 	OnGetIsInReverseDelegate = InArgs._IsInReverse;
 
@@ -44,16 +48,21 @@ void SLiveLinkHubPlaybackWidget::Construct(const FArguments& InArgs)
 	OnSetLoopingDelegate = InArgs._OnSetLooping;
 	OnGetLoopingDelegate = InArgs._IsLooping;
 
+	IsPlaybackEnabledDelegate = InArgs._IsPlaybackEnabled;
+
 	OnGetFrameRate = InArgs._GetFrameRate;
 
 	OnGetFrameBufferRanges = InArgs._GetBufferRanges;
 
 	const TAttribute<EFrameNumberDisplayFormats> GetDisplayFormatAttr = MakeAttributeSP(this, &SLiveLinkHubPlaybackWidget::GetDisplayFormat);
 	const TAttribute<FFrameRate> GetDisplayRateAttr = MakeAttributeSP(this, &SLiveLinkHubPlaybackWidget::GetFrameRate);
+	TAttribute<bool> IsPlaybackEnabled = MakeAttributeSP(this, &SLiveLinkHubPlaybackWidget::IsPlaybackEnabled);
 	
 	// Create our numeric type interface so we can pass it to the time slider below.
 	NumberInterface = MakeShareable(new FFrameNumberInterface(GetDisplayFormatAttr, 0, GetDisplayRateAttr, GetDisplayRateAttr));
 
+	const ISlateStyle* LiveLinkStyle = FSlateStyleRegistry::FindSlateStyle("LiveLinkStyle");
+
 	NumberInterface->DisplayFormatChanged();
 	ChildSlot
 	[
@@ -63,6 +72,7 @@ void SLiveLinkHubPlaybackWidget::Construct(const FArguments& InArgs)
 		.Padding(16.f, 8.f)
 		[
 			SNew(SLiveLinkHubTimeSlider)
+			.IsEnabled(IsPlaybackEnabled)
 			.BaseArgs(SSimpleTimeSlider::FArguments()
 			.ClampRangeHighlightSize(0.35f)
 			.ClampRangeHighlightColor(FLinearColor::Gray.CopyWithNewOpacity(0.5f))
@@ -96,6 +106,7 @@ void SLiveLinkHubPlaybackWidget::Construct(const FArguments& InArgs)
 			.AutoWidth()
 			[
 				SNew(SButton)
+				.IsEnabled(IsPlaybackEnabled)
 				.ButtonStyle(FAppStyle::Get(), "Animation.PlayControlsButton")
 				.OnClicked(this, &SLiveLinkHubPlaybackWidget::OnFirstFramePressed)
 				.ToolTipText(LOCTEXT("ToFront", "To Front"))
@@ -110,6 +121,7 @@ void SLiveLinkHubPlaybackWidget::Construct(const FArguments& InArgs)
 			.AutoWidth()
 			[
 				SNew(SButton)
+				.IsEnabled(IsPlaybackEnabled)
 				.ButtonStyle(FAppStyle::Get(), "Animation.PlayControlsButton")
 				.OnClicked(this, &SLiveLinkHubPlaybackWidget::OnPreviousFramePressed)
 				.ToolTipText(LOCTEXT("ToPrevious", "To Previous"))
@@ -124,6 +136,7 @@ void SLiveLinkHubPlaybackWidget::Construct(const FArguments& InArgs)
 			.AutoWidth()
 			[
 				SNew(SButton)
+				.IsEnabled(IsPlaybackEnabled)
 				.OnClicked(this, &SLiveLinkHubPlaybackWidget::OnPlayReversePressed)
 				.ToolTipText_Lambda([this]()
 				{
@@ -142,6 +155,7 @@ void SLiveLinkHubPlaybackWidget::Construct(const FArguments& InArgs)
 			.AutoWidth()
 			[
 				SNew(SButton)
+				.IsEnabled(IsPlaybackEnabled)
 				.OnClicked(this, &SLiveLinkHubPlaybackWidget::OnPlayForwardPressed)
 				.ToolTipText_Lambda([this]()
 				{
@@ -160,6 +174,7 @@ void SLiveLinkHubPlaybackWidget::Construct(const FArguments& InArgs)
 			.AutoWidth()
 			[
 				SNew(SButton)
+				.IsEnabled(IsPlaybackEnabled)
 				.ButtonStyle(FAppStyle::Get(), "Animation.PlayControlsButton")
 				.OnClicked(this, &SLiveLinkHubPlaybackWidget::OnNextFramePressed)
 				.ToolTipText(LOCTEXT("ToNext", "To Next"))
@@ -174,6 +189,7 @@ void SLiveLinkHubPlaybackWidget::Construct(const FArguments& InArgs)
 			.AutoWidth()
 			[
 				SNew(SButton)
+				.IsEnabled(IsPlaybackEnabled)
 				.ButtonStyle(FAppStyle::Get(), "Animation.PlayControlsButton")
 				.OnClicked(this, &SLiveLinkHubPlaybackWidget::OnLastFramePressed)
 				.ToolTipText(LOCTEXT("ToEnd", "To End"))
@@ -188,6 +204,7 @@ void SLiveLinkHubPlaybackWidget::Construct(const FArguments& InArgs)
 			.AutoWidth()
 			[
 				SNew(SButton)
+				.IsEnabled(IsPlaybackEnabled)
 				.ButtonStyle(FAppStyle::Get(), "Animation.PlayControlsButton")
 				.OnClicked(this, &SLiveLinkHubPlaybackWidget::OnLoopPressed)
 				.ToolTipText_Raw(this, &SLiveLinkHubPlaybackWidget::GetLoopTooltip)
@@ -198,10 +215,25 @@ void SLiveLinkHubPlaybackWidget::Construct(const FArguments& InArgs)
 					.Image(this, &SLiveLinkHubPlaybackWidget::GetLoopIcon)
 				]
 			]
+			+ SHorizontalBox::Slot()
+			.AutoWidth()
+			[
+				SNew(SButton)
+				.ButtonStyle(LiveLinkStyle, "LiveLinkHub.EjectButton")
+				.ContentPadding(FMargin(2.0f))
+				.OnClicked(this, &SLiveLinkHubPlaybackWidget::OnExitPlaybackPressed)
+				.ToolTipText(LOCTEXT("EjectButtonToolTip", "Exit Playback"))
+				[
+					SNew(SImage)
+						.Image(FAppStyle::Get().GetBrush("LevelViewport.EjectActorPilot.Small"))
+						.ColorAndOpacity(FSlateColor::UseSubduedForeground())
+				]
+			]
 			+SHorizontalBox::Slot()
 			.Padding(12.f, 0.f, 0.f, 0.f)
 			[
 				SNew(SSpinBox<double>)
+				.IsEnabled(IsPlaybackEnabled)
 				.Value(this, &SLiveLinkHubPlaybackWidget::GetSelectionStartTime)
 				.ToolTipText(LOCTEXT("SelectionStartTime", "Selection start time"))
 				.OnValueCommitted(this, &SLiveLinkHubPlaybackWidget::OnSelectionStartTimeCommitted)
@@ -218,6 +250,7 @@ void SLiveLinkHubPlaybackWidget::Construct(const FArguments& InArgs)
 			.Padding(4.f, 0.f)
 			[
 				SNew(SSpinBox<double>)
+				.IsEnabled(IsPlaybackEnabled)
 				.Value(this, &SLiveLinkHubPlaybackWidget::GetCurrentTime)
 				.ToolTipText(LOCTEXT("CurrentTime", "Current playback time"))
 				.OnValueCommitted(this, &SLiveLinkHubPlaybackWidget::OnCurrentTimeCommitted)
@@ -234,6 +267,7 @@ void SLiveLinkHubPlaybackWidget::Construct(const FArguments& InArgs)
 			.Padding(4.f, 0.f)
 			[
 				SNew(SSpinBox<double>)
+				.IsEnabled(IsPlaybackEnabled)
 				.Value(this, &SLiveLinkHubPlaybackWidget::GetSelectionEndTime)
 				.ToolTipText(LOCTEXT("SelectionEndTime", "Selection end time"))
 				.OnValueCommitted(this, &SLiveLinkHubPlaybackWidget::OnSelectionEndTimeCommitted)
@@ -250,6 +284,7 @@ void SLiveLinkHubPlaybackWidget::Construct(const FArguments& InArgs)
 			.Padding(4.f, 0.f)
 			[
 				SNew(SSpinBox<double>)
+				.IsEnabled(IsPlaybackEnabled)
 				.Value(this, &SLiveLinkHubPlaybackWidget::GetTotalLength)
 				.IsEnabled(false)
 				.ToolTipText(LOCTEXT("RecordingLength", "Recording length"))
@@ -263,6 +298,7 @@ void SLiveLinkHubPlaybackWidget::Construct(const FArguments& InArgs)
 			.AutoWidth()
 			[
 				SNew(SComboButton)
+				.IsEnabled(IsPlaybackEnabled)
 				.HasDownArrow(true)
 				.OnGetMenuContent(this, &SLiveLinkHubPlaybackWidget::MakePlaybackSettingsDropdown)
 				.ForegroundColor(FSlateColor::UseStyle())
@@ -316,6 +352,12 @@ FReply SLiveLinkHubPlaybackWidget::OnLoopPressed()
 	return FReply::Handled();
 }
 
+FReply SLiveLinkHubPlaybackWidget::OnExitPlaybackPressed()
+{
+	OnExitPlaybackDelegate.ExecuteIfBound();
+	return FReply::Handled();
+}
+
 void SLiveLinkHubPlaybackWidget::SetCurrentTime(double InTime)
 {
 	check(OnSetCurrentTimeDelegate.IsBound());
@@ -548,4 +590,9 @@ TSharedRef<SWidget> SLiveLinkHubPlaybackWidget::MakePlaybackSettingsDropdown()
 	return MenuBuilder.MakeWidget();
 }
 
+bool SLiveLinkHubPlaybackWidget::IsPlaybackEnabled() const
+{
+	return IsPlaybackEnabledDelegate.Execute();
+}
+
 #undef LOCTEXT_NAMESPACE
diff --git a/Engine/Plugins/Animation/LiveLinkHub/Source/LiveLinkHub/Private/UI/Widgets/SLiveLinkHubPlaybackWidget.h b/Engine/Plugins/Animation/LiveLinkHub/Source/LiveLinkHub/Private/UI/Widgets/SLiveLinkHubPlaybackWidget.h
--- a/Engine/Plugins/Animation/LiveLinkHub/Source/LiveLinkHub/Private/UI/Widgets/SLiveLinkHubPlaybackWidget.h
+++ b/Engine/Plugins/Animation/LiveLinkHub/Source/LiveLinkHub/Private/UI/Widgets/SLiveLinkHubPlaybackWidget.h
@@ -42,10 +42,14 @@ public:
 	SLATE_EVENT(FOnButtonPressed, OnPreviousFrame)
 	/** When next frame is pressed. */
 	SLATE_EVENT(FOnButtonPressed, OnNextFrame)
+	/** When the eject button is pressed. */
+	SLATE_EVENT(FOnButtonPressed, OnExitPlayback)
 	/** Is the recording looping? */
 	SLATE_EVENT(FOnGetChecked, IsLooping)
 	/** Set the loop state of the recording. */
 	SLATE_EVENT(FOnSetChecked, OnSetLooping)
+	/** Get if the playback widget should be enabled.*/
+	SLATE_EVENT(FOnGetChecked, IsPlaybackEnabled)
 	/** Is the recording paused? */
 	SLATE_EVENT(FOnGetChecked, IsPaused)
 	/** Is the recording playing in reverse? */
@@ -96,6 +100,8 @@ private:
 	FReply OnNextFramePressed();
 	/** When loop recording is pressed. */
 	FReply OnLoopPressed();
+	/** When recording is exited. */
+	FReply OnExitPlaybackPressed();
 
 	/** Sets the current playhead time. */
 	void SetCurrentTime(double InTime);
@@ -169,6 +175,9 @@ private:
 
 	/** Create the dropdown box for playback settings. */
 	TSharedRef<SWidget> MakePlaybackSettingsDropdown();
+
+	/** Returns whether the playback widgets should be enabled or not by calling the IsPlaybackEnabledDelegate delegate. */
+	bool IsPlaybackEnabled() const;
 private:
 	TWeakPtr<FLiveLinkHubPlaybackController> PlaybackController;
 	/** Delegate for pressing play forward. */
@@ -189,6 +198,12 @@ private:
 	/** Delegate for going to the next frame. */
 	FOnButtonPressed OnNextFrameDelegate;
 
+	/** Delegate for exiting the current recording. */
+	FOnButtonPressed OnExitPlaybackDelegate;
+
+	/** Delegate checking if the recording playback is enabled */
+	FOnGetChecked IsPlaybackEnabledDelegate;
+
 	/** Delegate for checking if the recording playback is paused. */
 	FOnGetChecked OnGetPausedDelegate;
 	
diff --git a/Engine/Plugins/Animation/LiveLinkHub/Source/LiveLinkHub/Public/LiveLinkHubApplicationBase.h b/Engine/Plugins/Animation/LiveLinkHub/Source/LiveLinkHub/Public/LiveLinkHubApplicationBase.h
--- a/Engine/Plugins/Animation/LiveLinkHub/Source/LiveLinkHub/Public/LiveLinkHubApplicationBase.h
+++ b/Engine/Plugins/Animation/LiveLinkHub/Source/LiveLinkHub/Public/LiveLinkHubApplicationBase.h
@@ -5,6 +5,20 @@
 #include "UObject/NameTypes.h"
 #include "WorkflowOrientedApp/WorkflowCentricApplication.h"
 
+namespace UE::LiveLinkHub
+{
+	static const FName LiveLinkStyleName = "LiveLinkStyle";
+
+	//~ Tab names to allow plugin code to refer to them.
+	static const FName SourcesTabId = "SourcesTabId";
+	static const FName SourceDetailsTabId = "SourceDetailsTabId";
+	static const FName SubjectsTabId = "SubjectsTabId";
+	static const FName PlaybackTabId = "PlaybackBottomTabId";
+	static const FName RecordingListTabId = "RecordingListTabId";
+	static const FName ClientsTabId = "ClientsTabId";
+	static const FName ClientDetailsTabId = "ClientDetailsTabId";
+}
+
 class FLiveLinkHubApplicationMode;
 
 /** Base class for the live link hub application, meant to expose a limited interface to modes that extend FLiveLinkHubApplicationMode. */
diff --git a/Engine/Plugins/VirtualProduction/CaptureManager/CaptureManagerApp/Source/CaptureManagerEditor/Private/UI/LiveLinkHubCaptureManagerMode.cpp b/Engine/Plugins/VirtualProduction/CaptureManager/CaptureManagerApp/Source/CaptureManagerEditor/Private/UI/LiveLinkHubCaptureManagerMode.cpp
--- a/Engine/Plugins/VirtualProduction/CaptureManager/CaptureManagerApp/Source/CaptureManagerEditor/Private/UI/LiveLinkHubCaptureManagerMode.cpp
+++ b/Engine/Plugins/VirtualProduction/CaptureManager/CaptureManagerApp/Source/CaptureManagerEditor/Private/UI/LiveLinkHubCaptureManagerMode.cpp
@@ -131,7 +131,7 @@ FLiveLinkHubCaptureManagerMode::FLiveLinkHubCaptureManagerMode(TSharedPtr<FLiveL
 {
 	using namespace UE::CaptureManager;
 
-	TabLayout = FTabManager::NewLayout("LiveLinkCaptureManagerMode_v1.1");
+	TabLayout = FTabManager::NewLayout("LiveLinkCaptureManagerMode_v1.2");
 	
 	TOptional<FString> LocalIpAddressOpt = GetLocalIpAddress();
 
@@ -161,6 +161,12 @@ FLiveLinkHubCaptureManagerMode::FLiveLinkHubCaptureManagerMode(TSharedPtr<FLiveL
 				->SetSizeCoefficient(1.f)
 				->SetOrientation(Orient_Horizontal)
 				->Split
+				(
+					FTabManager::NewStack()
+					->SetSizeCoefficient(0.5f)
+					->AddTab("LiveLinkDevices", ETabState::OpenedTab)
+				)
+				->Split
 				(
 					FTabManager::NewStack()
 					->SetSizeCoefficient(0.5f)
@@ -185,6 +191,12 @@ FLiveLinkHubCaptureManagerMode::FLiveLinkHubCaptureManagerMode(TSharedPtr<FLiveL
 					)
 				)
 			)
+			->Split
+			(
+				FTabManager::NewStack()
+				->SetHideTabWell(true)
+				->AddTab(UE::LiveLinkHub::PlaybackTabId, ETabState::ClosedTab)
+			)
 		);
 	
 	
