LiveLinkHub - Move Timecode Provider and Custom Time Step settings from the Timecode widget to the  Time & Sync settings section of LiveLinkHub

These settings will be saved in a per user setting file rather than a session config.

diff --git a/Engine/Plugins/Animation/LiveLink/Source/LiveLink/Private/LiveLinkCustomTimeStep.cpp b/Engine/Plugins/Animation/LiveLink/Source/LiveLink/Private/LiveLinkCustomTimeStep.cpp
--- a/Engine/Plugins/Animation/LiveLink/Source/LiveLink/Private/LiveLinkCustomTimeStep.cpp
+++ b/Engine/Plugins/Animation/LiveLink/Source/LiveLink/Private/LiveLinkCustomTimeStep.cpp
@@ -247,8 +247,10 @@ void ULiveLinkCustomTimeStep::OnLiveLinkClientUnregistered(const FName& Type, cl
 
 void ULiveLinkCustomTimeStep::OnLiveLinkSubjectAdded(FLiveLinkSubjectKey InSubjectKey)
 {
-	if (InSubjectKey == SubjectKey)
+	if (InSubjectKey.SubjectName == SubjectKey.SubjectName)
 	{
+		// Update the subject key, it's possible we didn't save with the same source ID.
+		SubjectKey = InSubjectKey;
 		RegisterLiveLinkSubject();
 	}
 }
@@ -256,7 +258,7 @@ void ULiveLinkCustomTimeStep::OnLiveLinkSubjectAdded(FLiveLinkSubjectKey InSubje
 
 void ULiveLinkCustomTimeStep::OnLiveLinkSubjectRemoved(FLiveLinkSubjectKey InSubjectKey)
 {
-	if (InSubjectKey == RegisteredSubjectKey)
+	if (InSubjectKey.SubjectName == RegisteredSubjectKey.SubjectName)
 	{
 		State = ECustomTimeStepSynchronizationState::Error;
 		UnregisterLiveLinkSubject();
diff --git a/Engine/Plugins/Animation/LiveLink/Source/LiveLink/Public/LiveLinkTimecodeProvider.h b/Engine/Plugins/Animation/LiveLink/Source/LiveLink/Public/LiveLinkTimecodeProvider.h
--- a/Engine/Plugins/Animation/LiveLink/Source/LiveLink/Public/LiveLinkTimecodeProvider.h
+++ b/Engine/Plugins/Animation/LiveLink/Source/LiveLink/Public/LiveLinkTimecodeProvider.h
@@ -63,6 +63,7 @@ public:
 	{
 		SubjectKey = InKey;
 	}
+
 private:
 	FQualifiedFrameTime ConvertTo(FQualifiedFrameTime Value) const;
 
@@ -118,4 +119,6 @@ private:
 	TArray<FLiveLinkTime> SubjectFrameTimes;
 	mutable FCriticalSection SubjectFrameLock; // Only lock SubjectFrameTimes
 	FDelegateHandle RegisterForFrameDataReceivedHandle;
+
+	friend struct FLiveLinkHubTimecodeSettings;
 };
diff --git a/Engine/Plugins/Animation/LiveLink/Source/LiveLinkEditor/Private/LiveLinkEditorModule.cpp b/Engine/Plugins/Animation/LiveLink/Source/LiveLinkEditor/Private/LiveLinkEditorModule.cpp
--- a/Engine/Plugins/Animation/LiveLink/Source/LiveLinkEditor/Private/LiveLinkEditorModule.cpp
+++ b/Engine/Plugins/Animation/LiveLink/Source/LiveLinkEditor/Private/LiveLinkEditorModule.cpp
@@ -116,6 +116,8 @@ public:
 
 		StyleSet->Set("LiveLinkHub.Icon.Small", new IMAGE_PLUGIN_BRUSH(TEXT("Starship/LiveLinkHub_16x16"), Icon16x16));
 		StyleSet->Set("LiveLinkHub.Icon", new IMAGE_PLUGIN_BRUSH_SVG(TEXT("Starship/LiveLinkHub_Icon"), Icon40x40));
+		StyleSet->Set("LiveLinkHub.EllipsisIcon", new IMAGE_PLUGIN_BRUSH_SVG(TEXT("Starship/Ellipsis"), Icon16x16));
+
 
 		StyleSet->Set("ClassIcon.LiveLinkPreset",                      new IMAGE_PLUGIN_BRUSH_SVG("Starship/LiveLink", Icon16x16));
 		StyleSet->Set("ClassIcon.LiveLinkFrameInterpolationProcessor", new IMAGE_PLUGIN_BRUSH_SVG("Starship/LiveLink", Icon16x16));
diff --git a/Engine/Plugins/Animation/LiveLinkHub/Source/LiveLinkHub/Private/Clients/LiveLinkHubProvider.cpp b/Engine/Plugins/Animation/LiveLinkHub/Source/LiveLinkHub/Private/Clients/LiveLinkHubProvider.cpp
--- a/Engine/Plugins/Animation/LiveLinkHub/Source/LiveLinkHub/Private/Clients/LiveLinkHubProvider.cpp
+++ b/Engine/Plugins/Animation/LiveLinkHub/Source/LiveLinkHub/Private/Clients/LiveLinkHubProvider.cpp
@@ -407,16 +407,14 @@ void FLiveLinkHubProvider::HandleHubConnectMessage(const FLiveLinkHubConnectMess
 	}
 
 	// Update the timecode provider when a client establishes connection.
-	if (const TSharedPtr<ILiveLinkHubSessionManager> Manager = SessionManager.Pin())
+	if (GetDefault<ULiveLinkHubTimeAndSyncSettings>()->bUseLiveLinkHubAsTimecodeSource)
 	{
-		if (const TSharedPtr<ILiveLinkHubSession> ActiveSession = Manager->GetCurrentSession())
-		{
-			if (ActiveSession->ShouldUseLiveLinkHubAsTimecodeSource())
-			{
-				SendTimecodeSettings(ActiveSession->GetTimecodeSettings(), *UpdatedClient);
-				SendCustomTimeStepSettings(ActiveSession->GetCustomTimeStepSettings(), *UpdatedClient);
-			}
-		}
+		SendTimecodeSettings(GetDefault<ULiveLinkHubTimeAndSyncSettings>()->TimecodeSettings, *UpdatedClient);
+	}
+
+	if (GetDefault<ULiveLinkHubTimeAndSyncSettings>()->bUseLiveLinkHubAsCustomTimeStepSource)
+	{
+		SendCustomTimeStepSettings(GetDefault<ULiveLinkHubTimeAndSyncSettings>()->CustomTimeStepSettings, *UpdatedClient);
 	}
 }
 	
@@ -435,19 +433,18 @@ void FLiveLinkHubProvider::HandleClientInfoMessage(const FLiveLinkClientInfoMess
 		}
 	}
 
-	if (const TSharedPtr<ILiveLinkHubSessionManager> Manager = SessionManager.Pin())
+	if (ClientId.IsValid())
 	{
-		if (const TSharedPtr<ILiveLinkHubSession> ActiveSession = Manager->GetCurrentSession())
+		if (GetDefault<ULiveLinkHubTimeAndSyncSettings>()->bUseLiveLinkHubAsTimecodeSource)
 		{
-			if (ActiveSession->ShouldUseLiveLinkHubAsTimecodeSource() && ClientId.IsValid())
-			{
-				SendTimecodeSettings(ActiveSession->GetTimecodeSettings(), ClientId);
-			}
+			SendTimecodeSettings(GetDefault<ULiveLinkHubTimeAndSyncSettings>()->TimecodeSettings, ClientId);
+		}
+
+		if (GetDefault<ULiveLinkHubTimeAndSyncSettings>()->bUseLiveLinkHubAsCustomTimeStepSource)
+		{
+			SendCustomTimeStepSettings(GetDefault<ULiveLinkHubTimeAndSyncSettings>()->CustomTimeStepSettings, ClientId);
 		}
-	}
 
-	if (ClientId.IsValid())
-	{
 		OnClientEventDelegate.Broadcast(ClientId, EClientEventType::Modified);
 	}
 }
@@ -701,32 +698,29 @@ void FLiveLinkHubProvider::SetClientEnabled(FLiveLinkHubClientId Client, bool bI
 
 	if (const TSharedPtr<ILiveLinkHubSessionManager> Manager = SessionManager.Pin())
 	{
-		if (const TSharedPtr<ILiveLinkHubSession> ActiveSession = Manager->GetCurrentSession())
+		if (GetDefault<ULiveLinkHubTimeAndSyncSettings>()->bUseLiveLinkHubAsTimecodeSource)
 		{
-			if (ActiveSession->ShouldUseLiveLinkHubAsTimecodeSource())
+			if (bInEnable)
 			{
-				if (bInEnable)
-				{
-					// Enabling client, send it up to date timecode and custom time step settings.
-					SendTimecodeSettings(ActiveSession->GetTimecodeSettings(), Client);
-				}
-				else
-				{
-					// Disabling it, so reset its timecode and custom time step.
-					ResetTimecodeSettings(Client);
-				}
+				// Enabling client, send it up to date timecode and custom time step settings.
+				SendTimecodeSettings(GetDefault<ULiveLinkHubTimeAndSyncSettings>()->TimecodeSettings, Client);
+			}
+			else
+			{
+				// Disabling it, so reset its timecode and custom time step.
+				ResetTimecodeSettings(Client);
 			}
+		}
 
-			if (ActiveSession->ShouldUseLiveLinkHubAsCustomTimeStepSource())
+		if (GetDefault<ULiveLinkHubTimeAndSyncSettings>()->bUseLiveLinkHubAsCustomTimeStepSource)
+		{
+			if (bInEnable)
 			{
-				if (bInEnable)
-				{
-					SendCustomTimeStepSettings(ActiveSession->GetCustomTimeStepSettings(), Client);
-				}
-				else
-				{
-					ResetCustomTimeStepSettings(Client);
-				}
+				SendCustomTimeStepSettings(GetDefault<ULiveLinkHubTimeAndSyncSettings>()->CustomTimeStepSettings, Client);
+			}
+			else
+			{
+				ResetCustomTimeStepSettings(Client);
 			}
 		}
 	}
diff --git a/Engine/Plugins/Animation/LiveLinkHub/Source/LiveLinkHub/Private/LiveLinkHub.cpp b/Engine/Plugins/Animation/LiveLinkHub/Source/LiveLinkHub/Private/LiveLinkHub.cpp
--- a/Engine/Plugins/Animation/LiveLinkHub/Source/LiveLinkHub/Private/LiveLinkHub.cpp
+++ b/Engine/Plugins/Animation/LiveLinkHub/Source/LiveLinkHub/Private/LiveLinkHub.cpp
@@ -31,6 +31,7 @@
 #include "Recording/LiveLinkHubRecordingListController.h"
 #include "Session/LiveLinkHubSessionManager.h"
 #include "Settings/LiveLinkHubSettings.h"
+#include "Settings/LiveLinkHubTimeAndSyncSettings.h"
 #include "Subjects/LiveLinkHubSubjectController.h"
 #include "UI/Window/LiveLinkHubWindowController.h"
 
@@ -189,6 +190,9 @@ void FLiveLinkHub::InitializePostEngineInit()
 
 	PlaybackController->Start();
 
+	GetDefault<ULiveLinkHubTimeAndSyncSettings>()->ApplyTimecodeProvider();
+	GetDefault<ULiveLinkHubTimeAndSyncSettings>()->ApplyCustomTimeStep();
+
 	GIsRunning = true;
 }
 
@@ -453,6 +457,12 @@ void FLiveLinkHub::RegisterLiveLinkHubSettings()
 			LOCTEXT("LiveLinkHubDescription", "Configure Live Link Hub."),
 			GetMutableDefault<ULiveLinkHubSettings>()
 		);
+
+		SettingsModule->RegisterSettings("Project", "Application", "Timing & Sync",
+			LOCTEXT("LiveLinkHubTimeAndSyncName", "Timing & Sync"),
+			LOCTEXT("LiveLinkHubTimeAndSyncNameDescription", "Configure Live Link Hub timecode and genlock settings."),
+			GetMutableDefault<ULiveLinkHubTimeAndSyncSettings>()
+		);
 	}
 }
 
@@ -463,6 +473,7 @@ void FLiveLinkHub::UnregisterLiveLinkHubSettings()
 		SettingsModule->UnregisterSettings("Editor", "Plugins", "Live Link");
 		SettingsModule->UnregisterSettings("Project", "Plugins", "Live Link");
 		SettingsModule->UnregisterSettings("Project", "Plugins", "Live Link Hub");
+		SettingsModule->UnregisterSettings("Project", "Application", "Timing & Sync");
 	}
 }
 
diff --git a/Engine/Plugins/Animation/LiveLinkHub/Source/LiveLinkHub/Private/Session/LiveLinkHubSession.h b/Engine/Plugins/Animation/LiveLinkHub/Source/LiveLinkHub/Private/Session/LiveLinkHubSession.h
--- a/Engine/Plugins/Animation/LiveLinkHub/Source/LiveLinkHub/Private/Session/LiveLinkHubSession.h
+++ b/Engine/Plugins/Animation/LiveLinkHub/Source/LiveLinkHub/Private/Session/LiveLinkHubSession.h
@@ -5,20 +5,22 @@
 #include "Async/Async.h"
 #include "Clients/LiveLinkHubProvider.h"
 #include "Clients/LiveLinkHubUEClientInfo.h"
+#include "Engine/SystemTimeTimecodeProvider.h"
 #include "Features/IModularFeatures.h"
 #include "Framework/Application/SlateApplication.h"
 #include "LiveLinkClient.h"
 #include "LiveLinkHubModule.h"
 #include "LiveLinkHubSessionData.h"
+#include "LiveLinkTimecodeProvider.h"
 #include "LiveLinkTypes.h"
 #include "LiveLinkVirtualSubject.h"
 #include "Misc/ScopeLock.h"
 #include "Modules/ModuleManager.h"
+#include "Settings/LiveLinkHubTimeAndSyncSettings.h"
 #include "UObject/Package.h"
 #include "UObject/StrongObjectPtr.h"
 #include "UObject/UObjectGlobals.h"
 
-
 #define LOCTEXT_NAMESPACE "LiveLinkHubSession"
 
 DECLARE_MULTICAST_DELEGATE_OneParam(FOnClientAddedToSession, FLiveLinkHubClientId);
@@ -47,30 +49,6 @@ public:
 	/** Get the list of clients in this session (The list of clients that can receive data from the hub) */
 	virtual TArray<FLiveLinkHubClientId> GetSessionClients() const = 0;
 
-	/** Returns whether livelinkhub should be used as a timecode source for connected clients. */
-	virtual bool ShouldUseLiveLinkHubAsTimecodeSource() const = 0;
-
-	/** Set whether livelinkhub should be used as a timecode source for connected clients.  */
-	virtual void SetUseLiveLinkHubAsTimecodeSource(bool bUseLiveLinkHubAsTimecodeSource) = 0;
-
-	/** Get the timecode settings for the current session. */
-	virtual FLiveLinkHubTimecodeSettings GetTimecodeSettings() const = 0;
-
-	/** Set the timecode settings for the current session. */
-	virtual void SetTimecodeSettings(const FLiveLinkHubTimecodeSettings& TimecodeSettings) const = 0;
-
-	/** Returns whether livelinkhub should be used as a CustomTimeStep source for connected clients. */
-	virtual bool ShouldUseLiveLinkHubAsCustomTimeStepSource() const = 0;
-
-	/** Set whether livelinkhub should be used as a timecode source for connected clients.  */
-	virtual void SetShouldUseLiveLinkHubAsCustomTimeStepSource(bool bUseLiveLinkHubAsCustomTimeStepSource) = 0;
-
-	/** Get the CustomTimeStep settings for the current session. */
-	virtual FLiveLinkHubCustomTimeStepSettings GetCustomTimeStepSettings() const = 0;
-
-	/** Set the CustomTimeStep settings for the current session. */
-	virtual void SetCustomTimeStepSettings(const FLiveLinkHubCustomTimeStepSettings& CustomTimeStepSettings) const = 0;
-
 	/** Get the topology mode for this instance. */
 	virtual ELiveLinkTopologyMode GetTopologyMode() const = 0;
 
@@ -114,10 +92,16 @@ public:
 					CachedSessionClients.Add(Client);
 				}
 
-				if (ShouldUseLiveLinkHubAsTimecodeSource())
+				const ULiveLinkHubTimeAndSyncSettings* Settings = GetDefault<ULiveLinkHubTimeAndSyncSettings>();
+
+				if (Settings->bUseLiveLinkHubAsTimecodeSource)
 				{
-					LiveLinkProvider->UpdateTimecodeSettings(GetTimecodeSettings(), Client);
-					LiveLinkProvider->UpdateCustomTimeStepSettings(GetCustomTimeStepSettings(), Client);
+					LiveLinkProvider->UpdateTimecodeSettings(Settings->TimecodeSettings, Client);
+				}
+
+				if (Settings->bUseLiveLinkHubAsCustomTimeStepSource)
+				{
+					LiveLinkProvider->UpdateCustomTimeStepSettings(Settings->CustomTimeStepSettings, Client);
 				}
 			}
 		}
@@ -183,54 +167,6 @@ public:
 		return CachedSessionClients.Contains(Client);
 	}
 
-	virtual bool ShouldUseLiveLinkHubAsTimecodeSource() const override
-	{
-		FReadScopeLock Locker(SessionDataLock);
-		return SessionData->bUseLiveLinkHubAsTimecodeSource;
-	}
-
-	virtual void SetUseLiveLinkHubAsTimecodeSource(bool bUseLiveLinkHubAsTimecodeSource) override
-	{
-		FWriteScopeLock Locker(SessionDataLock);
-		SessionData->bUseLiveLinkHubAsTimecodeSource = bUseLiveLinkHubAsTimecodeSource;
-	}
-
-	virtual FLiveLinkHubTimecodeSettings GetTimecodeSettings() const override
-	{
-		FReadScopeLock Locker(SessionDataLock);
-		return SessionData->TimecodeSettings;
-	}
-
-	virtual void SetTimecodeSettings(const FLiveLinkHubTimecodeSettings& TimecodeSettings) const override
-	{
-		FWriteScopeLock Locker(SessionDataLock);
-		SessionData->TimecodeSettings = TimecodeSettings;
-	}
-
-	virtual bool ShouldUseLiveLinkHubAsCustomTimeStepSource() const override
-	{
-		FReadScopeLock Locker(SessionDataLock);
-		return SessionData->bUseLiveLinkHubAsCustomTimeStepSource;
-	}
-
-	virtual void SetShouldUseLiveLinkHubAsCustomTimeStepSource(bool bUseLiveLinkHubAsCustomTimeStepSource) override
-	{
-		FWriteScopeLock Locker(SessionDataLock);
-		SessionData->bUseLiveLinkHubAsCustomTimeStepSource = bUseLiveLinkHubAsCustomTimeStepSource;
-	}
-
-	virtual FLiveLinkHubCustomTimeStepSettings GetCustomTimeStepSettings() const override
-	{
-		FReadScopeLock Locker(SessionDataLock);
-		return SessionData->CustomTimeStepSettings;
-	}
-
-	virtual void SetCustomTimeStepSettings(const FLiveLinkHubCustomTimeStepSettings& CustomTimeStepSettings) const override
-	{
-		FWriteScopeLock Locker(SessionDataLock);
-		SessionData->CustomTimeStepSettings = CustomTimeStepSettings;
-	}
-
 	virtual ELiveLinkTopologyMode GetTopologyMode() const override
 	{
 		FReadScopeLock Locker(SessionDataLock);
diff --git a/Engine/Plugins/Animation/LiveLinkHub/Source/LiveLinkHub/Private/Session/LiveLinkHubSessionData.h b/Engine/Plugins/Animation/LiveLinkHub/Source/LiveLinkHub/Private/Session/LiveLinkHubSessionData.h
--- a/Engine/Plugins/Animation/LiveLinkHub/Source/LiveLinkHub/Private/Session/LiveLinkHubSessionData.h
+++ b/Engine/Plugins/Animation/LiveLinkHub/Source/LiveLinkHub/Private/Session/LiveLinkHubSessionData.h
@@ -37,14 +37,6 @@ public:
 	UPROPERTY()
 	TArray<FLiveLinkHubUEClientInfo> Clients;
 
-	/** Timecode settings for the Live Link Hub. */
-	UPROPERTY()
-	FLiveLinkHubTimecodeSettings TimecodeSettings;
-
-	/** Frame Lock settings for the Live Link Hub. */
-	UPROPERTY()
-	FLiveLinkHubCustomTimeStepSettings CustomTimeStepSettings;
-
 	/** Recording metadata - session name */
 	UPROPERTY()
 	FString RecordingSessionName;
@@ -57,14 +49,6 @@ public:
 	UPROPERTY()
 	int32 RecordingTakeNumber;
 
-	/** Whether the hub should be used as a timecode source for connected clients. */
-	UPROPERTY()
-	bool bUseLiveLinkHubAsTimecodeSource = false;
-
-	/** Whether the hub should be used as a CustomTimeStep source for connected clients. */
-	UPROPERTY()
-	bool bUseLiveLinkHubAsCustomTimeStepSource = false;
-
 	/** If the hub should be a hub or a spoke. Hubs can only connect to UE instances, while spokes can connect to other hubs. */
 	UPROPERTY()
 	ELiveLinkTopologyMode TopologyMode = ELiveLinkTopologyMode::Hub;
diff --git a/Engine/Plugins/Animation/LiveLinkHub/Source/LiveLinkHub/Private/Session/LiveLinkHubSessionManager.h b/Engine/Plugins/Animation/LiveLinkHub/Source/LiveLinkHub/Private/Session/LiveLinkHubSessionManager.h
--- a/Engine/Plugins/Animation/LiveLinkHub/Source/LiveLinkHub/Private/Session/LiveLinkHubSessionManager.h
+++ b/Engine/Plugins/Animation/LiveLinkHub/Source/LiveLinkHub/Private/Session/LiveLinkHubSessionManager.h
@@ -401,12 +401,6 @@ private:
 				}
 			}
 
-			SessionData->TimecodeSettings.AssignTimecodeSettingsAsProviderToEngine();
-			LiveLinkProvider->UpdateTimecodeSettings(SessionData->TimecodeSettings);
-
-			SessionData->CustomTimeStepSettings.AssignCustomTimeStepToEngine();
-			LiveLinkProvider->UpdateCustomTimeStepSettings(SessionData->CustomTimeStepSettings);
-
 			TSharedPtr<FLiveLinkHubSession> CurrentSessionPtr;
 			{
 				FScopeLock Lock(&CurrentSessionCS);
@@ -431,6 +425,9 @@ private:
 
 			OnActiveSessionChangedDelegate.Broadcast(CurrentSessionPtr.ToSharedRef());
 
+			LiveLinkProvider->UpdateTimecodeSettings(GetDefault<ULiveLinkHubTimeAndSyncSettings>()->TimecodeSettings);
+			LiveLinkProvider->UpdateCustomTimeStepSettings(GetDefault<ULiveLinkHubTimeAndSyncSettings>()->CustomTimeStepSettings);
+
 			if (GetDefault<ULiveLinkHubSettings>()->bRemoveInvalidSubjectsAfterLoadingSession)
 			{
 				constexpr double ActualWaitTime = 1.5;
diff --git a/Engine/Plugins/Animation/LiveLinkHub/Source/LiveLinkHub/Private/Settings/LiveLinkHubTimeAndSyncSettings.h b/Engine/Plugins/Animation/LiveLinkHub/Source/LiveLinkHub/Private/Settings/LiveLinkHubTimeAndSyncSettings.h
new file mode 100644
--- /dev/null
+++ b/Engine/Plugins/Animation/LiveLinkHub/Source/LiveLinkHub/Private/Settings/LiveLinkHubTimeAndSyncSettings.h
@@ -0,0 +1,146 @@
+// Copyright Epic Games, Inc. All Rights Reserved.
+
+#pragma once
+
+#include "UObject/Object.h"
+
+#include "Clients/LiveLinkHubProvider.h"
+#include "Engine/Engine.h"
+#include "Engine/EngineCustomTimeStep.h"
+#include "Engine/TimecodeProvider.h"
+#include "Features/IModularFeatures.h"
+#include "ILiveLinkClient.h"
+#include "LiveLinkHub.h"
+#include "LiveLinkHubMessages.h"
+#include "Misc/FrameRate.h"
+#include "LiveLinkCustomTimeStep.h"
+#include "LiveLinkHubLog.h"
+#include "LiveLinkTimecodeProvider.h"
+
+#include "LiveLinkHubTimeAndSyncSettings.generated.h"
+
+/**
+ * Settings for LiveLinkHub's timecode and genlock.
+ */
+UCLASS(config=EditorPerProjectUserSettings)
+class ULiveLinkHubTimeAndSyncSettings : public UObject
+{
+	GENERATED_BODY()
+public:
+	UPROPERTY(config, EditAnywhere, Category = "Timecode")
+	FLiveLinkHubTimecodeSettings TimecodeSettings;
+
+	/** Whether the hub should be used as a timecode source for connected clients. */
+	UPROPERTY(config, EditAnywhere, Category = "Timecode", DisplayName = "Enable")
+	bool bUseLiveLinkHubAsTimecodeSource = false;
+
+	/** Custom time step */
+	UPROPERTY(config, EditAnywhere, Category = "Frame Lock")
+	FLiveLinkHubCustomTimeStepSettings CustomTimeStepSettings;
+
+	/** Whether the hub should be used as a CustomTimeStep source for connected clients. */
+	UPROPERTY(config, EditAnywhere, Category = "Frame Lock", DisplayName = "Enable")
+	bool bUseLiveLinkHubAsCustomTimeStepSource = false;
+
+public:
+#if WITH_EDITOR
+	//~ Begin UObject interface
+	virtual void PostEditChangeProperty(FPropertyChangedEvent& PropertyChangedEvent) override
+	{
+		const FName PropertyName = (PropertyChangedEvent.MemberProperty != nullptr) ? PropertyChangedEvent.MemberProperty->GetFName() : NAME_None;
+
+		if (PropertyName == GET_MEMBER_NAME_CHECKED(ULiveLinkHubTimeAndSyncSettings, TimecodeSettings))
+		{
+			if (bUseLiveLinkHubAsTimecodeSource)
+			{
+				OnToggleTimecodeSettings();
+			}
+		}
+		else if (PropertyName == GET_MEMBER_NAME_CHECKED(ULiveLinkHubTimeAndSyncSettings, CustomTimeStepSettings))
+		{
+			if (bUseLiveLinkHubAsCustomTimeStepSource)
+			{
+				OnToggleCustomTimeStepSettings();
+			}
+		}
+		else if (PropertyName == GET_MEMBER_NAME_CHECKED(ULiveLinkHubTimeAndSyncSettings, bUseLiveLinkHubAsTimecodeSource))
+		{
+			OnToggleTimecodeSettings();
+		}
+		else if (PropertyName == GET_MEMBER_NAME_CHECKED(ULiveLinkHubTimeAndSyncSettings, bUseLiveLinkHubAsCustomTimeStepSource))
+		{
+			OnToggleCustomTimeStepSettings();
+		}
+	}
+	//~ End UObject interface
+#endif
+
+	/** Returns whether the timecode configuration is valid. */
+	bool IsTimecodeProviderValid() const
+	{
+		if (TimecodeSettings.Source == ELiveLinkHubTimecodeSource::UseSubjectName)
+		{
+			ILiveLinkClient* LiveLinkClient = &IModularFeatures::Get().GetModularFeature<ILiveLinkClient>(ILiveLinkClient::ModularFeatureName);
+			return LiveLinkClient->IsSubjectValid(TimecodeSettings.SubjectName);
+		}
+		return true;
+	}
+
+	/** Returns whether the custom time step configuration is valid. */
+	bool IsCustomTimeStepValid() const
+	{
+		ILiveLinkClient* LiveLinkClient = &IModularFeatures::Get().GetModularFeature<ILiveLinkClient>(ILiveLinkClient::ModularFeatureName);
+		return LiveLinkClient->IsSubjectValid(CustomTimeStepSettings.SubjectName);
+	}
+
+	/** Apply this timecode configuration to this instance of Live Link Hub. */
+	void ApplyTimecodeProvider() const
+	{
+		TimecodeSettings.AssignTimecodeSettingsAsProviderToEngine();
+	}
+
+	/** Apply this custom time step configuration to this instance of Live Link Hub. */
+	void ApplyCustomTimeStep() const
+	{
+		CustomTimeStepSettings.AssignCustomTimeStepToEngine();
+	}
+
+private:
+	/** Handles broadcasting timecode settings when they're enabled/disabled. */
+	void OnToggleTimecodeSettings()
+	{
+		ApplyTimecodeProvider();
+
+		if (const TSharedPtr<FLiveLinkHubProvider> LiveLinkProvider = FLiveLinkHub::Get()->GetLiveLinkProvider())
+		{
+			if (bUseLiveLinkHubAsTimecodeSource)
+			{
+				LiveLinkProvider->UpdateTimecodeSettings(TimecodeSettings);
+			}
+			else
+			{
+				// If we're disabling LLH as a TimecodeProvider, reset the timecode provider on the engine.
+				LiveLinkProvider->ResetTimecodeSettings();
+			}
+		}
+	}
+
+	/** Handles broadcasting custom time step settings when they're enabled/disabled. */
+	void OnToggleCustomTimeStepSettings()
+	{
+		ApplyCustomTimeStep();
+
+		if (const TSharedPtr<FLiveLinkHubProvider> LiveLinkProvider = FLiveLinkHub::Get()->GetLiveLinkProvider())
+		{
+			if (bUseLiveLinkHubAsCustomTimeStepSource)
+			{
+				LiveLinkProvider->UpdateCustomTimeStepSettings(CustomTimeStepSettings);
+			}
+			else
+			{
+				// If we're disabling LLH as a CustomTimeStep source, reset the custom time step on the engine.
+				LiveLinkProvider->ResetCustomTimeStepSettings();
+			}
+		}
+	}
+};
diff --git a/Engine/Plugins/Animation/LiveLinkHub/Source/LiveLinkHub/Private/UI/Widgets/SLiveLinkHubSettings.h b/Engine/Plugins/Animation/LiveLinkHub/Source/LiveLinkHub/Private/UI/Widgets/SLiveLinkHubSettings.h
--- a/Engine/Plugins/Animation/LiveLinkHub/Source/LiveLinkHub/Private/UI/Widgets/SLiveLinkHubSettings.h
+++ b/Engine/Plugins/Animation/LiveLinkHub/Source/LiveLinkHub/Private/UI/Widgets/SLiveLinkHubSettings.h
@@ -14,6 +14,7 @@
 #include "ISettingsSection.h"
 #include "Modules/ModuleManager.h"
 #include "Settings/LiveLinkHubSettings.h"
+#include "Styling/ToolBarStyle.h"
 #include "Widgets/Docking/SDockTab.h"
 #include "Widgets/Input/SComboBox.h"
 #include "Widgets/SCompoundWidget.h"
@@ -51,14 +52,17 @@ public:
 			.SetDisplayName(LOCTEXT("LiveLinkHubSettingsTabLabel", "Settings"))
 			.SetIcon(FSlateIcon(FAppStyle::GetAppStyleSetName(), "LevelEditor.Tabs.Details"));
 
+		const FToolBarStyle& ToolbarStyle = FAppStyle::Get().GetWidgetStyle<FToolBarStyle>("SlimToolbar");
+
 		ChildSlot
 			[
 				SNew(SBox)
 				.Padding(FMargin(0, 4.0f, 4.0f, 4.0f))
 				[
 					SNew(SComboButton)
-					.ContentPadding(0.f)
-					.ComboButtonStyle(&FAppStyle::Get().GetWidgetStyle<FComboButtonStyle>("ToolbarComboButton"))
+					.ContentPadding(FMargin(2.0, 4.0))
+					.ButtonStyle(&ToolbarStyle.ButtonStyle)
+					.ComboButtonStyle(&ToolbarStyle.ComboButtonStyle)
 					.ForegroundColor(FSlateColor::UseStyle())
 					.ToolTipText(LOCTEXT("SettingsToolTip", "Settings"))
 					.ButtonContent()
diff --git a/Engine/Plugins/Animation/LiveLinkHub/Source/LiveLinkHub/Private/UI/Widgets/SLiveLinkTimecode.cpp b/Engine/Plugins/Animation/LiveLinkHub/Source/LiveLinkHub/Private/UI/Widgets/SLiveLinkTimecode.cpp
--- a/Engine/Plugins/Animation/LiveLinkHub/Source/LiveLinkHub/Private/UI/Widgets/SLiveLinkTimecode.cpp
+++ b/Engine/Plugins/Animation/LiveLinkHub/Source/LiveLinkHub/Private/UI/Widgets/SLiveLinkTimecode.cpp
@@ -2,183 +2,41 @@
 
 #include "SLiveLinkTimecode.h"
 
-#include "UObject/NameTypes.h"
-
 #include "EditorFontGlyphs.h"
-#include "Features/IModularFeatures.h"
-#include "Fonts/SlateFontInfo.h"
 #include "Framework/MultiBox/MultiBoxBuilder.h"
 #include "Internationalization/Text.h"
+#include "ISettingsModule.h"
 #include "Misc/App.h"
-#include "UI/Widgets/SLiveLinkTimecode.h"
-#include "Widgets/Input/SComboButton.h"
-#include "Widgets/Input/SNumericEntryBox.h"
+#include "Modules/ModuleManager.h"
+#include "Styling/SlateStyleRegistry.h"
+#include "Styling/ToolBarStyle.h"
+#include "Widgets/Images/SImage.h"
 #include "Widgets/Layout/SSpacer.h"
 #include "Widgets/STimecode.h"
-
-#include "Clients/LiveLinkHubProvider.h"
-#include "LiveLinkClient.h"
-#include "LiveLinkHub.h"
-#include "LiveLinkHubCommands.h"
-#include "LiveLinkHubMessages.h"
+ 
 #include "LiveLinkHubModule.h"
 #include "Session/LiveLinkHubSessionManager.h"
 
 #define LOCTEXT_NAMESPACE "LiveLinkHub"
 
-namespace UE::LiveLinkTimecode::Private
-{
-	static FName EnableTimecodeSourceId = TEXT("EnableTimeCodeSource");
-
-	/** We only support a preset list of timecode values + named subjects. */
-	static FName System23976fps = TEXT("SystemTime23976fps");
-	static FName System24fps = TEXT("SystemTime24fps");
-	static FName System25fps = TEXT("SystemTime25fps");
-	static FName System2997fps = TEXT("SystemTime2997fps");
-	static FName System30fps = TEXT("SystemTime30fps");
-	static FName System48fps = TEXT("SystemTime48fps");
-	static FName System50fps = TEXT("SystemTime50fps");
-	static FName System5994fps = TEXT("SystemTime5994fps");
-	static FName System60fps = TEXT("SystemTime60fps");
-
-	struct FTimecodeMenuItem
-	{
-		FFrameRate Rate;
-		FText Label;
-		FText ToolTip;
-	};
-	static const TMap<FName, FTimecodeMenuItem> StaticTimecodeMenu = []()
-	{
-		TMap<FName, FTimecodeMenuItem> Definition = {
-			{System23976fps, {
-				FFrameRate(24000,1001),
-				LOCTEXT("LiveLinkHubTimecodeSource23976fps", "System Time (23.976 ND fps)"),
-				LOCTEXT("LiveLinkHubTimecodeSource23976fps_Tooltip", "Use a 23.976 Non-drop FPS time code based on system time.")
-			}},
-			{System24fps, {
-				FFrameRate(24,1),
-				LOCTEXT("LiveLinkHubTimecodeSource24fps", "System Time (24 fps)"),
-				LOCTEXT("LiveLinkHubTimecodeSource24fps_Tooltip", "Use a 24 FPS time code based on system time.")
-			}},
-			{System25fps, {
-				FFrameRate(25,1),
-				LOCTEXT("LiveLinkHubTimecodeSource25fps", "System Time (25 fps)"),
-				LOCTEXT("LiveLinkHubTimecodeSource25fps_Tooltip", "Use a 25 FPS time code based on system time.")
-			}},
-			{System2997fps, {
-				FFrameRate(30000,1001),
-				LOCTEXT("LiveLinkHubTimecodeSource2997fps", "System Time (29.97 ND fps)"),
-				LOCTEXT("LiveLinkHubTimecodeSource29976fps_Tooltip", "Use a 29.97 Non-drop FPS time code based on system time.")
-			}},
-			{System30fps, {
-				FFrameRate(30,1),
-				LOCTEXT("LiveLinkHubTimecodeSource30fps", "System Time (30 fps)"),
-				LOCTEXT("LiveLinkHubTimecodeSource30fps_Tooltip", "Use a 30 FPS time code based on system time.")
-			}},
-			{System48fps, {
-				FFrameRate(48,1),
-				LOCTEXT("LiveLinkHubTimecodeSource48fps", "System Time (48 fps)"),
-				LOCTEXT("LiveLinkHubTimecodeSource48fps_Tooltip", "Use a 48 FPS time code based on system time.")
-			}},
-			{System50fps, {
-				FFrameRate(50,1),
-				LOCTEXT("LiveLinkHubTimecodeSource50fps", "System Time (50 fps)"),
-				LOCTEXT("LiveLinkHubTimecodeSource50fps_Tooltip", "Use a 50 FPS time code based on system time.")
-			}},
-			{System5994fps, {
-				FFrameRate(60000, 1001),
-				LOCTEXT("LiveLinkHubTimecodeSource5994fps", "System Time (59.94 ND fps)"),
-				LOCTEXT("LiveLinkHubTimecodeSource5994fps_Tooltip", "Use a 59.94 Non-drp FPS time code based on system time.")
-			}},
-			{System60fps, {
-				FFrameRate(60, 1),
-				LOCTEXT("LiveLinkHubTimecodeSource60fps", "System Time (60 fps)"),
-				LOCTEXT("LiveLinkHubTimecodeSource60fps_Tooltip", "Use a 60 FPS time code based on system time.")
-			}}
-		};
-		return Definition;
-	}();
-
-	/** We only support a preset list of timecode values + named subjects. */
-	static FName CustomTimeStep23976fps = TEXT("CustomTimeStepSystemTime23976fps");
-	static FName CustomTimeStep24fps = TEXT("CustomTimeStepSystemTime24fps");
-	static FName CustomTimeStep25fps = TEXT("CustomTimeStepSystemTime25fps");
-	static FName CustomTimeStep2997fps  = TEXT("CustomTimeStepSystemTime2997fps");
-	static FName CustomTimeStep30fps    = TEXT("CustomTimeStepSystemTime30fps");
-	static FName CustomTimeStep48fps = TEXT("CustomTimeStepSystemTime48fps");
-	static FName CustomTimeStep50fps = TEXT("CustomTimeStepSystemTime50fps");
-	static FName CustomTimeStep5994fps = TEXT("CustomTimeStepSystemTime5994fps");
-	static FName CustomTimeStep60fps = TEXT("CustomTimeStepSystemTime60fps");
-	static FName CustomTimeStep100fps = TEXT("CustomTimeStepSystemTime100fps");
-	static FName CustomTimeStep120fps = TEXT("CustomTimeStepSystemTime120fps");
-	static FName CustomTimeStepUseTimecodeProvider = TEXT("CustomTimeStepUseTimecodeProvider");
-
-	static const TMap<FName, FTimecodeMenuItem> StaticCustomTimeStepMenu = []()
-	{
-		TMap<FName, FTimecodeMenuItem> Definition = {
-			{ CustomTimeStep23976fps,
-             { FFrameRate(24000, 1001),
-				LOCTEXT("LiveLinkHubCustomTimeStep23976fps", "23.976 ND fps"),
-				LOCTEXT("LiveLinkHubCustomTimeStep23976fps_Tooltip", "Use a 23.976 Non-drop FPS timecode for CustomTimeStep.") } },
-			{ CustomTimeStep24fps,
-             { FFrameRate(24, 1),
-				LOCTEXT("LiveLinkHubCustomTimeStep24fps", "24 fps"),
-				LOCTEXT("LiveLinkHubCustomTimeStep24fps_Tooltip", "Use a 24 FPS timecode for CustomTimeStep.") }                 },
-			{ CustomTimeStep25fps,
-             { FFrameRate(25, 1),
-				LOCTEXT("LiveLinkHubCustomTimeStep25fps", "25 fps"),
-				LOCTEXT("LiveLinkHubCustomTimeStep25fps_Tooltip", "Use a 25 FPS timecode for CustomTimeStep.") }                 },
-			{ CustomTimeStep2997fps,
-             { FFrameRate(30000, 1001),
-				LOCTEXT("LiveLinkHubCustomTimeStep2997fps", "29.97 ND fps"),
-				LOCTEXT("LiveLinkHubCustomTimeStep29976fps_Tooltip", "Use a 29.97 Non-drop FPS timecode for CustomTimeStep.") }  },
-			{ CustomTimeStep30fps,
-             { FFrameRate(30, 1),
-				LOCTEXT("LiveLinkHubCustomTimeStep30fps", "30 fps"),
-				LOCTEXT("LiveLinkHubCustomTimeStep30fps_Tooltip", "Use a 30 FPS timecode for CustomTimeStep.") }                 },
-			{ CustomTimeStep48fps,
-             { FFrameRate(48, 1),
-				LOCTEXT("LiveLinkHubCustomTimeStep48fps", "48 fps"),
-				LOCTEXT("LiveLinkHubCustomTimeStep48fps_Tooltip", "Use a 48 FPS timecode for CustomTimeStep.") }                 },
-			{ CustomTimeStep50fps,
-             { FFrameRate(50, 1),
-				LOCTEXT("LiveLinkHubCustomTimeStep50fps", "50 fps"),
-				LOCTEXT("LiveLinkHubCustomTimeStep50fps_Tooltip", "Use a 50 FPS timecode for CustomTimeStep.") }                 },
-			{ CustomTimeStep5994fps,
-             { FFrameRate(60000, 1001),
-				LOCTEXT("LiveLinkHubCustomTimeStep5994fps", "59.94 ND fps"),
-				LOCTEXT("LiveLinkHubCustomTimeStep5994fps_Tooltip", "Use a 59.94 Non-drp FPS timecode for CustomTimeStep.") }    },
-			{ CustomTimeStep60fps,
-             { FFrameRate(60, 1),
-				LOCTEXT("LiveLinkHubCustomTimeStep60fps", "60 fps"),
-				LOCTEXT("LiveLinkHubCustomTimeStep60fps_Tooltip", "Use a 60 FPS timecode for CustomTimeStep.") }                 },
-			{ CustomTimeStep100fps,
-             { FFrameRate(100, 1),
-				LOCTEXT("LiveLinkHubCustomTimeStep100fps", "100 fps"),
-				LOCTEXT("LiveLinkHubCustomTimeStep100fps_Tooltip", "Use a 100 FPS timecode for CustomTimeStep.") }               },
-			{
-             CustomTimeStep120fps,    { FFrameRate(120, 1),
-				  LOCTEXT("LiveLinkHubCustomTimeStep120fps", "120 fps"),
-				  LOCTEXT("LiveLinkHubCustomTimeStep120fps_Tooltip", "Use a 120 FPS timecode for CustomTimeStep.") }
-		}};
-		return Definition;
-	}();
-
-}
-
 FSlateColor SLiveLinkTimecode::GetStatusColor() const
 {
-	if (bCachedIsCustomTimeStepEnabled)
+	const ULiveLinkHubTimeAndSyncSettings* TimeAndSyncSettinsgs = GetDefault<ULiveLinkHubTimeAndSyncSettings>();
+	const bool bTimecodeSourceEnabled = TimeAndSyncSettinsgs->bUseLiveLinkHubAsTimecodeSource;
+	const bool bCustomTimeStepEnabled = TimeAndSyncSettinsgs->bUseLiveLinkHubAsCustomTimeStepSource;
+
+	if (bCustomTimeStepEnabled)
 	{
-		if (ActiveSubjectCustomTimeStep.IsNone() || ActiveCustomTimeStep.IsNone())
+		// Get synchronization state? 
+		if (!TimeAndSyncSettinsgs->IsCustomTimeStepValid())
 		{
 			return FSlateColor(FColor::Red);
 		}
 	}
 
-	if (!bCachedIsTimecodeSource)
+	if (!bTimecodeSourceEnabled)
 	{
-		return FSlateColor(FColor::Yellow);
+		return FSlateColor(FLinearColor::Gray);
 	}
 
 	return FSlateColor(FColor::Green);
@@ -186,23 +44,27 @@ FSlateColor SLiveLinkTimecode::GetStatusColor() const
 
 FText SLiveLinkTimecode::GetTimecodeTooltip() const
 {
-	if (bCachedIsCustomTimeStepEnabled)
+	const ULiveLinkHubTimeAndSyncSettings* TimeAndSyncSettinsgs = GetDefault<ULiveLinkHubTimeAndSyncSettings>();
+	const bool bTimecodeSourceEnabled = TimeAndSyncSettinsgs->bUseLiveLinkHubAsTimecodeSource;
+	const bool bCustomTimeStepEnabled = TimeAndSyncSettinsgs->bUseLiveLinkHubAsCustomTimeStepSource;
+
+	if (bCustomTimeStepEnabled)
 	{
-		if (ActiveSubjectCustomTimeStep.IsNone() || ActiveCustomTimeStep.IsNone())
+		if (!TimeAndSyncSettinsgs->IsCustomTimeStepValid())
 		{
 			return LOCTEXT("LiveLinkTimeCode_CustomTimeStepError", "Connected clients will not be synchronized unless you select both a framerate and a subject for the Custom Time Step.");
 		}
 	}
 
-	if (bCachedIsTimecodeSource && !bCachedIsCustomTimeStepEnabled)
+	if (bTimecodeSourceEnabled && !bCustomTimeStepEnabled)
 	{
 		return LOCTEXT("LiveLinkTimeCode_TimecodeConnected", "Sending Timecode data to connected editors.");
 	}
-	else if (!bCachedIsTimecodeSource && bCachedIsCustomTimeStepEnabled)
+	else if (!bTimecodeSourceEnabled && bCustomTimeStepEnabled)
 	{
 		return LOCTEXT("LiveLinkTimeCode_CustomTimeStepConnected", "Sending CustomTimestep data to connected editors.");
 	}
-	else if (bCachedIsTimecodeSource && bCachedIsCustomTimeStepEnabled)
+	else if (bTimecodeSourceEnabled && bCustomTimeStepEnabled)
 	{
 		return LOCTEXT("LiveLinkTimeCode_TimecodeAndCustomTimeStepConnected", "Sending Timecode and CustomTimeStep data to connected editors.");
 	}
@@ -210,547 +72,76 @@ FText SLiveLinkTimecode::GetTimecodeTooltip() const
 	return LOCTEXT("LiveLinkTimeCode_NotConnected", "No Timecode or CustomTimeStep data shared with connected editors.");
 }
 
-void SLiveLinkTimecode::OnEnableTimecodeToggled()
-{
-	bCachedIsTimecodeSource = !bCachedIsTimecodeSource;
-
-	if (TSharedPtr<ILiveLinkHubSessionManager> SessionManager = FModuleManager::Get().GetModuleChecked<FLiveLinkHubModule>("LiveLinkHub").GetSessionManager())
-	{
-		SessionManager->GetCurrentSession()->SetUseLiveLinkHubAsTimecodeSource(bCachedIsTimecodeSource);
-	}
-
-	if (bCachedIsTimecodeSource)
-	{
-		SendUpdatedTimecodeToEditor(MakeTimecodeSettings());
-	}
-	else
-	{
-		// If we're disabling LLH as 
-		const TSharedPtr<FLiveLinkHub> LiveLinkHub = FModuleManager::Get().GetModuleChecked<FLiveLinkHubModule>("LiveLinkHub").GetLiveLinkHub();
-		if (LiveLinkHub.IsValid())
-		{
-			if (TSharedPtr<FLiveLinkHubProvider> Provider = LiveLinkHub->GetLiveLinkProvider())
-			{
-				Provider->ResetTimecodeSettings();
-			}
-		}
-	}
-}
-
-void SLiveLinkTimecode::OnEnableCustomTimeStepToggled()
-{
-	bCachedIsCustomTimeStepEnabled = !bCachedIsCustomTimeStepEnabled;
-
-	if (TSharedPtr<ILiveLinkHubSessionManager> SessionManager = FModuleManager::Get().GetModuleChecked<FLiveLinkHubModule>("LiveLinkHub").GetSessionManager())
-	{
-		SessionManager->GetCurrentSession()->SetShouldUseLiveLinkHubAsCustomTimeStepSource(bCachedIsCustomTimeStepEnabled);
-	}
-
-	if (bCachedIsCustomTimeStepEnabled)
-	{
-		SendUpdatedCustomTimeStepToEditor(MakeCustomTimeStepSettings());
-	}
-	else
-	{
-		// If we're disabling LLH as a CustomTimeStep source, reset the custom time step on the engine.
-		const TSharedPtr<FLiveLinkHub> LiveLinkHub = FModuleManager::Get().GetModuleChecked<FLiveLinkHubModule>("LiveLinkHub").GetLiveLinkHub();
-		if (LiveLinkHub.IsValid())
-		{
-			if (TSharedPtr<FLiveLinkHubProvider> Provider = LiveLinkHub->GetLiveLinkProvider())
-			{
-				Provider->ResetCustomTimeStepSettings();
-			}
-		}
-	}
-}
-
-void SLiveLinkTimecode::SendUpdatedCustomTimeStepToEditor(const FLiveLinkHubCustomTimeStepSettings& CustomTimeStepSettings)
-{
-	const TSharedPtr<FLiveLinkHub> LiveLinkHub = FModuleManager::Get().GetModuleChecked<FLiveLinkHubModule>("LiveLinkHub").GetLiveLinkHub();
-	if (LiveLinkHub.IsValid())
-	{
-		if (TSharedPtr<FLiveLinkHubProvider> Provider = LiveLinkHub->GetLiveLinkProvider())
-		{
-			Provider->UpdateCustomTimeStepSettings(CustomTimeStepSettings);
-		}
-	}
-}
-
-FLiveLinkHubTimecodeSettings SLiveLinkTimecode::MakeTimecodeSettings() const
-{
-	FLiveLinkHubTimecodeSettings Settings;
-	using namespace UE::LiveLinkTimecode::Private;
-
-	auto GetSystemTimeSource = [this]() -> TOptional<FLiveLinkHubTimecodeSettings>
-		{
-			for (const TPair<FName, FTimecodeMenuItem>& Source : StaticTimecodeMenu)
-			{
-				if (ActiveTimecodeSource == Source.Key)
-				{
-					FLiveLinkHubTimecodeSettings Settings;
-					Settings.Source = ELiveLinkHubTimecodeSource::SystemTimeEditor;
-					Settings.DesiredFrameRate = Source.Value.Rate;
-					return Settings;
-				}
-			}
-			return {};
-		};
-
-	if (TOptional<FLiveLinkHubTimecodeSettings> SystemTimeSource = GetSystemTimeSource())
-	{
-		Settings = MoveTemp(*SystemTimeSource);
-	}
-	else
-	{
-		Settings.Source = ELiveLinkHubTimecodeSource::UseSubjectName;
-	}
-
-	Settings.SubjectName = ActiveTimecodeSource;
-
-	return Settings;
-}
-
-FLiveLinkHubCustomTimeStepSettings SLiveLinkTimecode::MakeCustomTimeStepSettings() const
-{
-	FLiveLinkHubCustomTimeStepSettings Settings;
-	Settings.SubjectName = ActiveSubjectCustomTimeStep;
-
-	using namespace UE::LiveLinkTimecode::Private;
-	for (const TPair<FName, FTimecodeMenuItem>& Source : StaticCustomTimeStepMenu)
-	{
-		if (ActiveCustomTimeStep == Source.Key)
-		{
-			Settings.CustomTimeStepRate = Source.Value.Rate;
-			break;
-		}
-	};
-
-	Settings.FrameRateDivider = CachedFrameRateDivider;
-	Settings.bLockStepMode = bCachedLockStepMode;
-
-	return Settings;
-}
-
-void SLiveLinkTimecode::SendUpdatedTimecodeToEditor(const FLiveLinkHubTimecodeSettings& TimecodeSettings)
-{
-	const TSharedPtr<FLiveLinkHub> LiveLinkHub = FModuleManager::Get().GetModuleChecked<FLiveLinkHubModule>("LiveLinkHub").GetLiveLinkHub();	
-	if (LiveLinkHub.IsValid())
-	{
-		if (TSharedPtr<FLiveLinkHubProvider> Provider = LiveLinkHub->GetLiveLinkProvider())
-		{
-			Provider->UpdateTimecodeSettings(TimecodeSettings);
-		}
-	}
-}
-
-void SLiveLinkTimecode::UpdateTimecodeSettingsFromSession(const TSharedRef<ILiveLinkHubSession>& InSession)
-{
-	using namespace UE::LiveLinkTimecode::Private;
-
-	const FLiveLinkHubTimecodeSettings Settings = InSession->GetTimecodeSettings();
-	bCachedIsTimecodeSource = InSession->ShouldUseLiveLinkHubAsTimecodeSource();
-
-	ActiveTimecodeSource = Settings.SubjectName.IsNone() ? System24fps : Settings.SubjectName;
-}
-
-void SLiveLinkTimecode::UpdateCustomTimeStepSettingsFromSession(const TSharedRef<ILiveLinkHubSession>& InSession)
-{
-	using namespace UE::LiveLinkTimecode::Private;
-
-	const FLiveLinkHubCustomTimeStepSettings Settings = InSession->GetCustomTimeStepSettings();
-	bCachedIsCustomTimeStepEnabled = InSession->ShouldUseLiveLinkHubAsCustomTimeStepSource();
-	bCachedLockStepMode = Settings.bLockStepMode;
-
-	const TPair<FName, FTimecodeMenuItem>* TargetItem = Algo::FindByPredicate(StaticCustomTimeStepMenu, [Rate = Settings.CustomTimeStepRate](const TPair<FName, FTimecodeMenuItem>& Source) { return Rate == Source.Value.Rate;  });
-	if (TargetItem)
-	{
-		ActiveCustomTimeStep = TargetItem->Key;
-	}
-
-	CachedFrameRateDivider = Settings.FrameRateDivider;
-	ActiveSubjectCustomTimeStep = Settings.SubjectName;
-}
-
-void SLiveLinkTimecode::SetCustomTimeStepSubject(const FName SubjectName)
-{
-	if (SubjectName == ActiveSubjectCustomTimeStep)
-	{
-		ActiveSubjectCustomTimeStep = NAME_None;
-	}
-	else
-	{
-		ActiveSubjectCustomTimeStep = SubjectName;
-	}
-
-	const FLiveLinkHubCustomTimeStepSettings Settings = MakeCustomTimeStepSettings();
-
-	const TSharedPtr<FLiveLinkHub> LiveLinkHub = FModuleManager::Get().GetModuleChecked<FLiveLinkHubModule>("LiveLinkHub").GetLiveLinkHub();
-	if (LiveLinkHub.IsValid())
-	{
-		LiveLinkHub->GetSessionManager()->GetCurrentSession()->SetCustomTimeStepSettings(Settings);
-	}
-
-	if (bCachedIsCustomTimeStepEnabled)
-	{
-		SendUpdatedCustomTimeStepToEditor(Settings);
-	}
-}
-
-void SLiveLinkTimecode::SetCustomTimeStepRateDivider(uint32 RateDivider)
-{
-	CachedFrameRateDivider = RateDivider;
-
-	const FLiveLinkHubCustomTimeStepSettings Settings = MakeCustomTimeStepSettings();
-	
-	const TSharedPtr<FLiveLinkHub> LiveLinkHub = FModuleManager::Get().GetModuleChecked<FLiveLinkHubModule>("LiveLinkHub").GetLiveLinkHub();
-	if (LiveLinkHub.IsValid())
-	{
-		LiveLinkHub->GetSessionManager()->GetCurrentSession()->SetCustomTimeStepSettings(Settings);
-	}
-
-	if (bCachedIsCustomTimeStepEnabled)
-	{
-		SendUpdatedCustomTimeStepToEditor(Settings);
-	}
-}
-
-void SLiveLinkTimecode::SetCustomTimeStep(const FName SourceId)
-{
-	if (SourceId == ActiveCustomTimeStep)
-	{
-		return;
-	}
-
-	ActiveCustomTimeStep = SourceId;
-
-	const FLiveLinkHubCustomTimeStepSettings Settings = MakeCustomTimeStepSettings();
-
-	const TSharedPtr<FLiveLinkHub> LiveLinkHub = FModuleManager::Get().GetModuleChecked<FLiveLinkHubModule>("LiveLinkHub").GetLiveLinkHub();
-	if (LiveLinkHub.IsValid())
-	{
-		LiveLinkHub->GetSessionManager()->GetCurrentSession()->SetCustomTimeStepSettings(Settings);
-	}
-
-	if (bCachedIsCustomTimeStepEnabled)
-	{
-		SendUpdatedCustomTimeStepToEditor(Settings);
-	}
-}
-
-void SLiveLinkTimecode::SetTimecodeSource(const FName SourceId)
-{
-	if (SourceId == ActiveTimecodeSource)
-	{
-		return;
-	}
-
-	ActiveTimecodeSource = SourceId;
-	
-	const FLiveLinkHubTimecodeSettings Settings = MakeTimecodeSettings();
-	Settings.AssignTimecodeSettingsAsProviderToEngine();
-
-	const TSharedPtr<FLiveLinkHub> LiveLinkHub = FModuleManager::Get().GetModuleChecked<FLiveLinkHubModule>("LiveLinkHub").GetLiveLinkHub();
-	if (LiveLinkHub.IsValid())
-	{
-		LiveLinkHub->GetSessionManager()->GetCurrentSession()->SetTimecodeSettings(Settings);
-	}
-
-	if (bCachedIsTimecodeSource)
-	{
-		SendUpdatedTimecodeToEditor(Settings);
-	}
-}
-
-TSharedRef<SWidget> SLiveLinkTimecode::MakeFrameRateDivider()
+void SLiveLinkTimecode::Construct(const FArguments& InArgs)
 {
-	FMenuBuilder MenuBuilder(/*bInShouldCloseWindowAfterMenuSelection=*/true, nullptr);
+	const FLiveLinkHubModule& LiveLinkHubModule = FModuleManager::Get().GetModuleChecked<FLiveLinkHubModule>("LiveLinkHub");
+	const ISlateStyle* StyleSet = FSlateStyleRegistry::FindSlateStyle("LiveLinkStyle");
+	check(StyleSet);
 
-	return SNew(SHorizontalBox)
+	ChildSlot
+	[
+		SNew(SHorizontalBox)
+		.ToolTipText(this, &SLiveLinkTimecode::GetTimecodeTooltip)
 		+ SHorizontalBox::Slot()
-		.FillWidth(1.f)
+		.AutoWidth()
+		.VAlign(VAlign_Center)
+		.Padding(5, 0, 3, 0)
 		[
-			SNew(SSpacer)
+			SNew(STextBlock)
+				.Font(FAppStyle::Get().GetFontStyle("FontAwesome.8"))
+				.ColorAndOpacity(this, &SLiveLinkTimecode::GetStatusColor)
+				.Text(FEditorFontGlyphs::Circle)
 		]
 		+ SHorizontalBox::Slot()
 		.AutoWidth()
-		.Padding(0.f, 0.f, 15.f, 0.f)
+		.VAlign(VAlign_Center)
+		.MinWidth(75)
+		.Padding(FMargin(2, 0))
 		[
-			SNew(SBox)
-				.WidthOverride(50.f)
-				[
-					SNew(SNumericEntryBox<uint32>)
-						.AllowSpin(true)
-						.MinDesiredValueWidth(10)
-						.MinValue(1)
-						.MinSliderValue(1)
-						.MaxSliderValue(100)
-						.OnValueChanged_Lambda([this](uint32 NewValue)
-						{
-							SetCustomTimeStepRateDivider(NewValue);
-						})
-						.OnEndSliderMovement_Lambda([this](uint32 NewValue)
-						{
-							CachedFrameRateDivider = NewValue;
-						})
-					.Value_Lambda([this]()
-						{
-							return CachedFrameRateDivider;
-						})
-				]
-		];
-}
-
-void SLiveLinkTimecode::MakeTimecodeMenuItems(FMenuBuilder& MenuBuilder)
-{
-	MenuBuilder.BeginSection("LiveLinkHub.Timecode.TimecodeProvider", LOCTEXT("TimecodeProviderSection", "Timecode Provider"));
-
-	auto GenerateUIAction = [this](const FName Id) -> FUIAction
-		{
-			// Generate a UI Action
-			return FUIAction(
-				FExecuteAction::CreateSP(this, &SLiveLinkTimecode::SetTimecodeSource, Id),
-				FCanExecuteAction::CreateLambda([] { return true; }),
-				FIsActionChecked::CreateLambda([this, Id] { return Id == ActiveTimecodeSource; }));
-		};
-	using namespace UE::LiveLinkTimecode::Private;
-	for (const TPair<FName, FTimecodeMenuItem>& Source : StaticTimecodeMenu)
-	{
-		MenuBuilder.AddMenuEntry(
-			Source.Value.Label,
-			Source.Value.ToolTip,
-			FSlateIcon(),
-			GenerateUIAction(Source.Key),
-			NAME_None,
-			EUserInterfaceActionType::Check);
-	}
-
-	constexpr bool bIncludeDisabledSubject = false;
-    constexpr bool bIncludeVirtualSubject = true;
-	TArray<FLiveLinkSubjectKey> Subjects = WorkingClient->GetSubjects(bIncludeDisabledSubject, bIncludeVirtualSubject);
-
-
-	for (const FLiveLinkSubjectKey& Subject : Subjects)
-	{
-		FTextBuilder SubjectText;
-		SubjectText.AppendLineFormat(LOCTEXT("LiveLinkHubTimecodeSourceSubject_Tooltip", "{0}'s timecode"), FText::FromName(Subject.SubjectName));
-		MenuBuilder.AddMenuEntry(
-			FText::FromName(Subject.SubjectName),
-			SubjectText.ToText(),
-			FSlateIcon(),
-			GenerateUIAction(Subject.SubjectName),
-			NAME_None,
-			EUserInterfaceActionType::Check);
-	}
-	MenuBuilder.EndSection();
-}
-
-
-void SLiveLinkTimecode::MakeCustomTimeStepMenuItems(FMenuBuilder& MenuBuilder)
-{
-	auto GenerateUIAction = [this](const FName Id) -> FUIAction
-		{
-			// Generate a UI Action
-			return FUIAction(
-				FExecuteAction::CreateSP(this, &SLiveLinkTimecode::SetCustomTimeStep, Id),
-				FCanExecuteAction::CreateLambda([] { return true; }),
-				FIsActionChecked::CreateLambda([this, Id] { return Id == ActiveCustomTimeStep; }));
-		};
-	using namespace UE::LiveLinkTimecode::Private;
-	for (const TPair<FName, FTimecodeMenuItem>& Source : StaticCustomTimeStepMenu)
-	{
-		MenuBuilder.AddMenuEntry(
-			Source.Value.Label,
-			Source.Value.ToolTip,
-			FSlateIcon(),
-			GenerateUIAction(Source.Key),
-			NAME_None,
-			EUserInterfaceActionType::Check);
-	}
-
-	MenuBuilder.AddMenuSeparator();
-
-	MenuBuilder.AddMenuEntry(
-		LOCTEXT("LiveLinkCustomTimeStepStep", "Lock Step Mode"),
-		LOCTEXT("LiveLinkCustomTimeStepUseTimecode_Tooltip", "Enable / disable lock step mode."),
-		FSlateIcon(),
-		FUIAction(FExecuteAction::CreateLambda(
-			[this]() mutable
+			SNew(STimecode)
+			.DisplayLabel(false)
+			.TimecodeFont(FCoreStyle::Get().GetFontStyle(TEXT("NormalText")))
+			.LabelColor(this, &SLiveLinkTimecode::GetStatusColor)
+			.TimecodeColor(this, &SLiveLinkTimecode::GetStatusColor)
+			.Timecode(MakeAttributeLambda([]
 			{
-				bCachedLockStepMode = !bCachedLockStepMode;
-			}),
-			FCanExecuteAction::CreateLambda([] { return true; }),
-			FIsActionChecked::CreateLambda([this] { return bCachedLockStepMode; })),
-		NAME_None,
-		EUserInterfaceActionType::Check);
-
-	MenuBuilder.AddWidget(MakeFrameRateDivider(), LOCTEXT("FrameRateDivider", "Rate Divider"), false);
-
-	FNewMenuDelegate SubMenuDelegate = FNewMenuDelegate::CreateLambda(
-		[this](FMenuBuilder& SubMenuBuilder)
-		{
-			constexpr bool bIncludeDisabledSubject = false;
-			constexpr bool bIncludeVirtualSubject = true;
-
-			TArray<FLiveLinkSubjectKey> Subjects = WorkingClient->GetSubjects(bIncludeDisabledSubject, bIncludeVirtualSubject);
-
-			for (const FLiveLinkSubjectKey& Subject : Subjects)
-			{
-				FTextBuilder SubjectText;
-				SubjectText.AppendLineFormat(
-					LOCTEXT("LiveLinkHubCustomTimeStepSourceSubject_Tooltip", "Use {0} as a frame lock source."),
-					FText::FromName(Subject.SubjectName)
-				);
-
-				FUIAction Action = FUIAction(
-					FExecuteAction::CreateSP(this, &SLiveLinkTimecode::SetCustomTimeStepSubject, Subject.SubjectName.Name),
-					FCanExecuteAction::CreateLambda([] { return true; }),
-					FIsActionChecked::CreateLambda([this, Subject] { return Subject.SubjectName == ActiveSubjectCustomTimeStep; }));
-
-				SubMenuBuilder.AddMenuEntry(
-					FText::FromName(Subject.SubjectName),
-					SubjectText.ToText(),
-					FSlateIcon(),
-					Action,
-					NAME_None,
-					EUserInterfaceActionType::Check);
-			}
-		});
-
-	MenuBuilder.AddSubMenu(
-		LOCTEXT("LiveLinkHubCustomTimeStepSubjectKey", "Subject Key"),
-		LOCTEXT("LiveLinkHubTimestepSubjectKey_Tooltip", "The subject that's used for synchronizing the engine."),
-		SubMenuDelegate,
-		FUIAction(),
-		NAME_None,
-		EUserInterfaceActionType::None);
-}
-
-void SLiveLinkTimecode::OnLockStepModeToggled()
-{
-	bCachedLockStepMode = !bCachedLockStepMode;
-
-	const FLiveLinkHubCustomTimeStepSettings Settings = MakeCustomTimeStepSettings();
-
-	const TSharedPtr<FLiveLinkHub> LiveLinkHub = FModuleManager::Get().GetModuleChecked<FLiveLinkHubModule>("LiveLinkHub").GetLiveLinkHub();
-	if (LiveLinkHub.IsValid())
-	{
-		LiveLinkHub->GetSessionManager()->GetCurrentSession()->SetCustomTimeStepSettings(Settings);
-	}
-
-	if (bCachedIsCustomTimeStepEnabled)
-	{
-		SendUpdatedCustomTimeStepToEditor(Settings);
-	}
-}
-
-TSharedRef<SWidget> SLiveLinkTimecode::MakeMenu()
-{
-	FMenuBuilder MenuBuilder(/*bInShouldCloseWindowAfterMenuSelection=*/true, nullptr);
-
-	FNewMenuDelegate TimecodeSubMenuDelegate = FNewMenuDelegate::CreateLambda(
-		[this](FMenuBuilder& SubMenuBuilder)
-		{
-			MakeTimecodeMenuItems(SubMenuBuilder);
-		});
-
-	MenuBuilder.AddSubMenu(
-		LOCTEXT("LiveLinkHubTimecodeSource", "Enable Timecode Source"),
-		LOCTEXT("LiveLinkHubTimecodeSource_Tooltip", "Make this Live Link Hub a time code source for connected editors."),
-		TimecodeSubMenuDelegate,
-		FUIAction(
-			FExecuteAction::CreateSP(this, &SLiveLinkTimecode::OnEnableTimecodeToggled),
-			FCanExecuteAction::CreateLambda([] { return true; }),
-			FIsActionChecked::CreateLambda([this] { return bCachedIsTimecodeSource; })),
-		NAME_None,
-		EUserInterfaceActionType::ToggleButton);
-
-
-	FNewMenuDelegate SubMenuDelegate = FNewMenuDelegate::CreateLambda(
-		[this](FMenuBuilder& SubMenuBuilder)
-		{
-			MakeCustomTimeStepMenuItems(SubMenuBuilder);
-		});
-
-	MenuBuilder.AddSubMenu(
-		LOCTEXT("LiveLinkHubCustomTimeStep", "Enable Custom Time Step"),
-		LOCTEXT("LiveLinkHubTimestep_Tooltip", "Enable CustomTimeStep forces the framerate of connected clients to the choosen value.."),
-		SubMenuDelegate,
-		FUIAction(
-			FExecuteAction::CreateSP(this, &SLiveLinkTimecode::OnEnableCustomTimeStepToggled),
-			FCanExecuteAction::CreateLambda([] { return true; }),
-			FIsActionChecked::CreateLambda([this] { return bCachedIsCustomTimeStepEnabled; })),
-		NAME_None,
-		EUserInterfaceActionType::ToggleButton);
-
-	return MenuBuilder.MakeWidget();
+				return FApp::GetTimecode();
+			}))
+		]
+		+ SHorizontalBox::Slot()
+		.AutoWidth()
+		.VAlign(VAlign_Center)
+		.Padding(FMargin(0, 3, 0, 2))
+		[
+			SNew(STextBlock)
+			.Font(FCoreStyle::Get().GetFontStyle(TEXT("SmallText")))
+			.ColorAndOpacity(FSlateColor::UseSubduedForeground())
+			.Text(MakeAttributeLambda([] {return FApp::GetTimecodeFrameRate().ToPrettyText();}))
+		]
+		+ SHorizontalBox::Slot()
+		.AutoWidth()
+		.VAlign(VAlign_Center)
+		[
+			SNew(SButton)
+			.ContentPadding(0.0)
+			.ButtonStyle(&FAppStyle::Get().GetWidgetStyle<FToolBarStyle>("SlimToolbar").ButtonStyle)
+			.OnClicked(this, &SLiveLinkTimecode::OnClickOpenSettings)
+			.ToolTipText(LOCTEXT("TimeAndSyncToolTip", "Open Time and Sync settings"))
+			[
+				SNew(SImage)
+				.Image(StyleSet->GetBrush("LiveLinkHub.EllipsisIcon"))
+			]
+		]
+	];
 }
 
-void SLiveLinkTimecode::Construct(const FArguments& InArgs)
+FReply SLiveLinkTimecode::OnClickOpenSettings() const
 {
-	WorkingClient = (FLiveLinkClient*)&IModularFeatures::Get().GetModularFeature<ILiveLinkClient>(ILiveLinkClient::ModularFeatureName);
-	
-	const FLiveLinkHubModule& LiveLinkHubModule = FModuleManager::Get().GetModuleChecked<FLiveLinkHubModule>("LiveLinkHub");
-	TSharedPtr<ILiveLinkHubSessionManager> SessionManager = LiveLinkHubModule.GetSessionManager();
-	if (SessionManager.IsValid())
+	if (ISettingsModule* SettingsModule = FModuleManager::GetModulePtr<ISettingsModule>("Settings"))
 	{
-		UpdateTimecodeSettingsFromSession(SessionManager->GetCurrentSession().ToSharedRef());
-		UpdateCustomTimeStepSettingsFromSession(SessionManager->GetCurrentSession().ToSharedRef());
-
-		SessionManager->OnActiveSessionChanged().AddLambda([this](const TSharedRef<ILiveLinkHubSession>& InSession)
-		{
-			// Update the UI when a config is loaded.
-			UpdateTimecodeSettingsFromSession(InSession);
-			UpdateCustomTimeStepSettingsFromSession(InSession);
-		});
+		SettingsModule->ShowViewer("Project", "Application", "Timing & Sync");
 	}
-
-	check(WorkingClient);
-	ChildSlot
-	[
-		SNew(SComboButton)
-		.ContentPadding(FMargin(4.0f, 0.0f))
-		.MenuPlacement(MenuPlacement_AboveAnchor)
-		.OnGetMenuContent(this, &SLiveLinkTimecode::MakeMenu)
-		.ComboButtonStyle(&FAppStyle::Get().GetWidgetStyle<FComboButtonStyle>("SimpleComboButton"))
-		.HasDownArrow(true)
-		.ToolTipText(this, &SLiveLinkTimecode::GetTimecodeTooltip)
-		.ButtonContent()
-		[
-			SNew(SHorizontalBox)
-			+ SHorizontalBox::Slot()
-			.AutoWidth()
-			.VAlign(VAlign_Center)
-			.Padding(0, 0, 3, 0)
-			[
-				SNew(STextBlock)
-					.Font(FAppStyle::Get().GetFontStyle("FontAwesome.8"))
-					.ColorAndOpacity(this, &SLiveLinkTimecode::GetStatusColor)
-					.Text(FEditorFontGlyphs::Circle)
-			]
-			+ SHorizontalBox::Slot()
-			.AutoWidth()
-			.VAlign(VAlign_Center)
-			.Padding(2, 0, 5, 0)
-			[
-				SNew(STimecode)
-				.DisplayLabel(false)
-				.TimecodeFont(FCoreStyle::Get().GetFontStyle(TEXT("NormalText")))
-				.Timecode(MakeAttributeLambda([]
-				{
-					return FApp::GetTimecode();
-				}))
-			]
-			+ SHorizontalBox::Slot()
-			.AutoWidth()
-			.VAlign(VAlign_Center)
-			.Padding(2, 0, 5, 0)
-			[
-				SNew(STextBlock)
-				.Font(FCoreStyle::Get().GetFontStyle(TEXT("NormalText")))
-				.Text(MakeAttributeLambda([] {return FApp::GetTimecodeFrameRate().ToPrettyText();}))
-			]
-    	]
-	];
+	return FReply::Handled();
 }
 
 #undef LOCTEXT_NAMESPACE
diff --git a/Engine/Plugins/Animation/LiveLinkHub/Source/LiveLinkHub/Private/UI/Widgets/SLiveLinkTimecode.h b/Engine/Plugins/Animation/LiveLinkHub/Source/LiveLinkHub/Private/UI/Widgets/SLiveLinkTimecode.h
--- a/Engine/Plugins/Animation/LiveLinkHub/Source/LiveLinkHub/Private/UI/Widgets/SLiveLinkTimecode.h
+++ b/Engine/Plugins/Animation/LiveLinkHub/Source/LiveLinkHub/Private/UI/Widgets/SLiveLinkTimecode.h
@@ -4,12 +4,6 @@
 #include "Widgets/SCompoundWidget.h"
 #include "UObject/NameTypes.h"
 
-class FName;
-class FLiveLinkClient;
-struct FLiveLinkHubTimecodeSettings;
-struct FLiveLinkHubCustomTimeStepSettings;
-class FMenuBuilder;
-
 /**
  * Utility class to display timecode status on the menu bar. Provides user with
  * a drop down of supported time code values including subjects that are in the
@@ -35,78 +29,6 @@ private:
 	/** Gets the tooltip text for the time code widget */
 	FText GetTimecodeTooltip() const;
 
-	/** Make the menu entry that corresponds to the FrameRateDivider property on the CustomTimeStep */
-	TSharedRef<SWidget> MakeFrameRateDivider();
-
-	/** Make menu items for the frame lock submenu */
-	void MakeCustomTimeStepMenuItems(FMenuBuilder& MenuBuilder);
-
-	/** Make the menu items for the timecode provider menu. */
-	void MakeTimecodeMenuItems(FMenuBuilder& MenuBuilder);
-
-	/** Makes the combo button menu for Timecode source */
-	TSharedRef<SWidget> MakeMenu();
-
-	/** Invoked by slate when the time code setting is enabled. */
-	void OnEnableTimecodeToggled();
-
-	/** Invoked by slate when the frame lock setting is enabled. */
-	void OnEnableCustomTimeStepToggled();
-
-	/** Invoked by slate when the lock step mode setting is toggle. */
-	void OnLockStepModeToggled();
-
-	/** Set the new frame rate divider.*/
-	void SetCustomTimeStepRateDivider(uint32 Divider);
-
-	/** Set the subject name to use with frame lock setting.*/
-	void SetCustomTimeStepSubject(const FName SubjectName);
-
-	/** Specifies the source for timecode */
-	void SetTimecodeSource(const FName SourceId);
-
-	/** Specifies the source for CustomTimeStep */
-	void SetCustomTimeStep(const FName SourceId);
-
-	/** Sends the current timecode settings to the connected editor(s). */
-	void SendUpdatedTimecodeToEditor(const FLiveLinkHubTimecodeSettings& TimecodeSettings);
-
-	/** Create a timecode settings objects from the session config. */
-	FLiveLinkHubTimecodeSettings MakeTimecodeSettings() const;
-
-	/** Create a CustomTimeStep settings objects from the session config. */
-	FLiveLinkHubCustomTimeStepSettings MakeCustomTimeStepSettings() const;
-
-	/** Sends the current frame lock settings to the connected editor(s). */
-	void SendUpdatedCustomTimeStepToEditor(const FLiveLinkHubCustomTimeStepSettings& CustomTimeStepSettings);
-
-	/** Update the local timecode properties from the timecode settings. */
-	void UpdateTimecodeSettingsFromSession(const TSharedRef<class ILiveLinkHubSession>& InSession);
-
-	/** Update the local CustomTimeStep properties from the session CustomTimeStep settings. */
-	void UpdateCustomTimeStepSettingsFromSession(const TSharedRef<ILiveLinkHubSession>& InSession);
-	
-	/** Do we have timecode source enabled */
-	bool bCachedIsTimecodeSource = false;
-
-	/** Do we have frame lock enabled. */
-	bool bCachedIsCustomTimeStepEnabled = false;
-
-	/** Name of the active timecode source. */
-	FName ActiveTimecodeSource;
-
-	/** Name of the active custom time step framerate. */
-	FName ActiveCustomTimeStep;
-
-	/** Enable / disable lock step mode for frame lock */
-	bool bCachedLockStepMode = true;
-
-	/** Frame rate divider for frame lock mode. */
-	uint32 CachedFrameRateDivider = 1;
-
-	/** The name of the subject applied to the frame lock on Live Link Custom Timestep. */
-	FName ActiveSubjectCustomTimeStep;
-
-	/** Cached value for the FLiveLinkClient */
-	FLiveLinkClient* WorkingClient = nullptr;
+	/** Open project settings to the Time & Sync section. */
+	FReply OnClickOpenSettings() const;
 };
diff --git a/Engine/Plugins/Animation/LiveLinkHub/Source/LiveLinkHubMessaging/Private/LiveLinkHubMessages.cpp b/Engine/Plugins/Animation/LiveLinkHub/Source/LiveLinkHubMessaging/Private/LiveLinkHubMessages.cpp
--- a/Engine/Plugins/Animation/LiveLinkHub/Source/LiveLinkHubMessaging/Private/LiveLinkHubMessages.cpp
+++ b/Engine/Plugins/Animation/LiveLinkHub/Source/LiveLinkHubMessaging/Private/LiveLinkHubMessages.cpp
@@ -8,7 +8,6 @@
 #include "Features/IModularFeatures.h"
 #include "ILiveLinkClient.h"
 #include "ILiveLinkModule.h"
-#include "LiveLinkTimecodeProvider.h"
 
 const FName UE::LiveLinkHub::Private::LiveLinkHubProviderType = TEXT("LiveLinkHub");
 FName FLiveLinkHubMessageAnnotation::ProviderTypeAnnotation = TEXT("ProviderType");
@@ -70,6 +69,8 @@ void FLiveLinkHubCustomTimeStepSettings::AssignCustomTimeStepToEngine() const
 
 	ULiveLinkHubCustomTimeStep* NewCustomTimeStep = NewObject<ULiveLinkHubCustomTimeStep>(GEngine, ULiveLinkHubCustomTimeStep::StaticClass());
 	NewCustomTimeStep->LiveLinkDataRate = CustomTimeStepRate;
+	NewCustomTimeStep->bLockStepMode = bLockStepMode;
+	NewCustomTimeStep->FrameRateDivider = FrameRateDivider;
 
 	if (TOptional<FLiveLinkSubjectKey> Target = UE::LiveLinkHubMessages::Private::FindSubjectKey(LiveLinkClient, SubjectName))
 	{
@@ -77,7 +78,8 @@ void FLiveLinkHubCustomTimeStepSettings::AssignCustomTimeStepToEngine() const
 	}
 	else
 	{
-		NewCustomTimeStep->SubjectKey = FLiveLinkSubjectKey{};
+		// Note: We must set the subject name because the livelink custom time step will use it to try to match new subjects being added.
+		NewCustomTimeStep->SubjectKey = FLiveLinkSubjectKey{FGuid(), SubjectName};
 	}
 
 	// Override the custom timestep for the engine.
@@ -101,7 +103,7 @@ void FLiveLinkHubTimecodeSettings::AssignTimecodeSettingsAsProviderToEngine() co
 		FName ObjectName = MakeUniqueObjectName(GEngine, USystemTimeTimecodeProvider::StaticClass(), "DefaultTimecodeProvider");
 		USystemTimeTimecodeProvider* NewTimecodeProvider = NewObject<USystemTimeTimecodeProvider>(GEngine, ObjectName);
 		NewTimecodeProvider->FrameRate = DesiredFrameRate;
-		NewTimecodeProvider->FrameDelay = 0;
+		NewTimecodeProvider->FrameDelay = FrameDelay;
 		GEngine->SetTimecodeProvider(NewTimecodeProvider);
 		UE_LOG(LogLiveLinkHubMessages, Display, TEXT("System Time Timecode provider set."));
 	}
@@ -114,6 +116,10 @@ void FLiveLinkHubTimecodeSettings::AssignTimecodeSettingsAsProviderToEngine() co
 			FName ObjectName = MakeUniqueObjectName(GEngine, ULiveLinkTimecodeProvider::StaticClass(), "DefaultLiveLinkTimecodeProvider");
 			ULiveLinkTimecodeProvider* LiveLinkProvider = NewObject<ULiveLinkTimecodeProvider>(GEngine, ObjectName);
 			LiveLinkProvider->SetTargetSubjectKey(*Target);
+			LiveLinkProvider->OverrideFrameRate = DesiredFrameRate;
+			LiveLinkProvider->FrameDelay = FrameDelay;
+			LiveLinkProvider->BufferSize = BufferSize;
+
 			GEngine->SetTimecodeProvider(LiveLinkProvider);
 			UE_LOG(LogLiveLinkHubMessages, Display, TEXT("Live Link Timecode provider assigned to %s."), *SubjectName.ToString());
 		}
diff --git a/Engine/Plugins/Animation/LiveLinkHub/Source/LiveLinkHubMessaging/Public/LiveLinkHubMessages.h b/Engine/Plugins/Animation/LiveLinkHub/Source/LiveLinkHubMessaging/Public/LiveLinkHubMessages.h
--- a/Engine/Plugins/Animation/LiveLinkHub/Source/LiveLinkHubMessaging/Public/LiveLinkHubMessages.h
+++ b/Engine/Plugins/Animation/LiveLinkHub/Source/LiveLinkHubMessaging/Public/LiveLinkHubMessages.h
@@ -6,6 +6,7 @@
 #include "ILiveLinkHubMessagingModule.h"
 #include "ILiveLinkClient.h"
 #include "LiveLinkMessages.h"
+#include "LiveLinkTimecodeProvider.h"
 #include "LiveLinkTypes.h"
 #include "Misc/FrameRate.h"
 
@@ -62,19 +63,19 @@ struct LIVELINKHUBMESSAGING_API FLiveLinkHubCustomTimeStepSettings
 	bool bResetCustomTimeStep = false;
 
 	/** Corresponds to the lock step mode in ULiveLinkCustomTimeStep. */
-	UPROPERTY()
+	UPROPERTY(config, EditAnywhere, Category = "Frame Lock")
 	bool bLockStepMode = true;
 
 	/** Corresponds to the frame rate divider in ULiveLinkCustomTimeStep */
-	UPROPERTY()
+	UPROPERTY(config, EditAnywhere, Category = "Frame Lock")
 	uint32 FrameRateDivider = 1;
 
 	/** If we are locking the editor frame rate to the subject then this property holds that subject name. */
-	UPROPERTY()
-	FName SubjectName;
+	UPROPERTY(config, EditAnywhere, Category = "Frame Lock")
+	FLiveLinkSubjectName SubjectName;
 
 	/** Desired frame rate to lock the editor. This corresponds to the LiveLinkDataRate in ULiveLinkCustomTimeStep */
-	UPROPERTY()
+	UPROPERTY(config, EditAnywhere, Category = "Frame Lock")
 	FFrameRate CustomTimeStepRate = FFrameRate(60, 1);
 
 	/** Assign the frame lock settings to the engine. */
@@ -101,16 +102,31 @@ struct LIVELINKHUBMESSAGING_API FLiveLinkHubTimecodeSettings
 	GENERATED_BODY()
 
 	/** Source time code value.  If it is not defined then we use the default time code provider in the engine. */
-	UPROPERTY()
+	UPROPERTY(config, EditAnywhere, Category = "Timecode")
 	ELiveLinkHubTimecodeSource Source = ELiveLinkHubTimecodeSource::NotDefined;
 
 	/** Name of the subject to map timecode if Source == ELiveLinkHubTimecodeSource::UseSubjectName */
-	UPROPERTY()
-	FName SubjectName;
+	UPROPERTY(config, EditAnywhere, Category = "Timecode")
+	FLiveLinkSubjectName SubjectName;
 
 	/** Desired frame rate to set if Source == ELiveLinkHubTimecodeSource::SystemTimeEditor. */
-	UPROPERTY()
-	FFrameRate DesiredFrameRate;
+	UPROPERTY(config, EditAnywhere, Category = "Timecode", meta = (EditCondition = "Source==ELiveLinkHubTimecodeSource::SystemTimeEditor"))
+	FFrameRate DesiredFrameRate = FFrameRate(60, 1);
+
+	/**
+	 * Number of frames to subtract from the qualified frame time when GetDelayedQualifiedFrameTime or GetDelayedTimecode is called.
+	 * @see GetDelayedQualifiedFrameTime, GetDelayedTimecode
+	 */
+	UPROPERTY(config, EditAnywhere, Category = "Timecode")
+	float FrameDelay = 0.f;
+
+	/** The number of frame to keep in memory. The provider will not be synchronized until the buffer is full at least once. */
+	UPROPERTY(config, EditAnywhere, Category = "Timecode")
+	int32 BufferSize;
+
+	/** How timecode should be evaluated. */
+	UPROPERTY(config, EditAnywhere, Category = "Timecode")
+	ELiveLinkTimecodeProviderEvaluationType EvaluationType;
 
 	/** Assign the settings to a new timecode provider and override the current engine settings. */
 	void AssignTimecodeSettingsAsProviderToEngine() const;
diff --git a/Engine/Source/Programs/LiveLinkHub/Config/DefaultEngine.ini b/Engine/Source/Programs/LiveLinkHub/Config/DefaultEngine.ini
--- a/Engine/Source/Programs/LiveLinkHub/Config/DefaultEngine.ini
+++ b/Engine/Source/Programs/LiveLinkHub/Config/DefaultEngine.ini
@@ -24,6 +24,7 @@ bAllowEditingRebroadcastProperty=false
 AutoConnectClients=LocalOnly
 +ProjectSettingsToDisplay=Live Link
 +ProjectSettingsToDisplay=Live Link Hub
++ProjectSettingsToDisplay=Timing & Sync
 +ProjectSettingsToDisplay=UDPMessaging
 
 [/Script/AppleARKit.AppleARKitSettings]
@@ -63,4 +64,3 @@ NodeMemoryLoggingThresholdMB=5
 
 [CoreRedirects]
 +PropertyRedirects=(OldName="LiveLinkHubSettings.bAutoAddDiscoveredClients",NewName="AutoConnectClients")
-
