LiveLinkHub - Clamp Timecode's FrameDelay and BufferSize
Clamp CustomTimeStep's FrameRateDivider
Fix Timecode not showing as invalid if the subject doesn't exist.
Fix TC/GL settings not getting reset correctly on the hub

diff --git a/Engine/Plugins/Animation/LiveLinkHub/Source/LiveLinkHub/Private/Settings/LiveLinkHubTimeAndSyncSettings.h b/Engine/Plugins/Animation/LiveLinkHub/Source/LiveLinkHub/Private/Settings/LiveLinkHubTimeAndSyncSettings.h
--- a/Engine/Plugins/Animation/LiveLinkHub/Source/LiveLinkHub/Private/Settings/LiveLinkHubTimeAndSyncSettings.h
+++ b/Engine/Plugins/Animation/LiveLinkHub/Source/LiveLinkHub/Private/Settings/LiveLinkHubTimeAndSyncSettings.h
@@ -12,10 +12,11 @@
 #include "ILiveLinkClient.h"
 #include "LiveLinkHub.h"
 #include "LiveLinkHubMessages.h"
-#include "Misc/FrameRate.h"
 #include "LiveLinkCustomTimeStep.h"
 #include "LiveLinkHubLog.h"
 #include "LiveLinkTimecodeProvider.h"
+#include "Misc/FrameRate.h"
+#include "UnrealEngine.h"
 
 #include "LiveLinkHubTimeAndSyncSettings.generated.h"
 
@@ -109,17 +110,17 @@ private:
 	/** Handles broadcasting timecode settings when they're enabled/disabled. */
 	void OnToggleTimecodeSettings()
 	{
-		ApplyTimecodeProvider();
-
 		if (const TSharedPtr<FLiveLinkHubProvider> LiveLinkProvider = FLiveLinkHub::Get()->GetLiveLinkProvider())
 		{
 			if (bUseLiveLinkHubAsTimecodeSource)
 			{
+				ApplyTimecodeProvider();
 				LiveLinkProvider->UpdateTimecodeSettings(TimecodeSettings);
 			}
 			else
 			{
 				// If we're disabling LLH as a TimecodeProvider, reset the timecode provider on the engine.
+				GEngine->Exec(GEngine->GetCurrentPlayWorld(nullptr), TEXT("TimecodeProvider.reset"));
 				LiveLinkProvider->ResetTimecodeSettings();
 			}
 		}
@@ -128,17 +129,17 @@ private:
 	/** Handles broadcasting custom time step settings when they're enabled/disabled. */
 	void OnToggleCustomTimeStepSettings()
 	{
-		ApplyCustomTimeStep();
-
 		if (const TSharedPtr<FLiveLinkHubProvider> LiveLinkProvider = FLiveLinkHub::Get()->GetLiveLinkProvider())
 		{
 			if (bUseLiveLinkHubAsCustomTimeStepSource)
 			{
+				ApplyCustomTimeStep();
 				LiveLinkProvider->UpdateCustomTimeStepSettings(CustomTimeStepSettings);
 			}
 			else
 			{
 				// If we're disabling LLH as a CustomTimeStep source, reset the custom time step on the engine.
+				GEngine->Exec(GEngine->GetCurrentPlayWorld(nullptr), TEXT("CustomTimeStep.reset"));
 				LiveLinkProvider->ResetCustomTimeStepSettings();
 			}
 		}
diff --git a/Engine/Plugins/Animation/LiveLinkHub/Source/LiveLinkHub/Private/UI/Widgets/SLiveLinkTimecode.cpp b/Engine/Plugins/Animation/LiveLinkHub/Source/LiveLinkHub/Private/UI/Widgets/SLiveLinkTimecode.cpp
--- a/Engine/Plugins/Animation/LiveLinkHub/Source/LiveLinkHub/Private/UI/Widgets/SLiveLinkTimecode.cpp
+++ b/Engine/Plugins/Animation/LiveLinkHub/Source/LiveLinkHub/Private/UI/Widgets/SLiveLinkTimecode.cpp
@@ -34,12 +34,14 @@ FSlateColor SLiveLinkTimecode::GetStatusColor() const
 		}
 	}
 
-	if (!bTimecodeSourceEnabled)
+	if (bTimecodeSourceEnabled)
+	{
+		return TimeAndSyncSettinsgs->IsTimecodeProviderValid() ? FSlateColor(FColor::Green) : FSlateColor(FColor::Red);
+	}
+	else
 	{
 		return FSlateColor(FLinearColor::Gray);
 	}
-
-	return FSlateColor(FColor::Green);
 }
 
 FText SLiveLinkTimecode::GetTimecodeTooltip() const
diff --git a/Engine/Plugins/Animation/LiveLinkHub/Source/LiveLinkHubMessaging/Public/LiveLinkHubMessages.h b/Engine/Plugins/Animation/LiveLinkHub/Source/LiveLinkHubMessaging/Public/LiveLinkHubMessages.h
--- a/Engine/Plugins/Animation/LiveLinkHub/Source/LiveLinkHubMessaging/Public/LiveLinkHubMessages.h
+++ b/Engine/Plugins/Animation/LiveLinkHub/Source/LiveLinkHubMessaging/Public/LiveLinkHubMessages.h
@@ -67,7 +67,7 @@ struct LIVELINKHUBMESSAGING_API FLiveLinkHubCustomTimeStepSettings
 	bool bLockStepMode = true;
 
 	/** Corresponds to the frame rate divider in ULiveLinkCustomTimeStep */
-	UPROPERTY(config, EditAnywhere, Category = "Frame Lock")
+	UPROPERTY(config, EditAnywhere, Category = "Frame Lock", meta = (ClampMin = 1, ClampMax = 256, UIMin = 1, UIMax = 256))
 	uint32 FrameRateDivider = 1;
 
 	/** If we are locking the editor frame rate to the subject then this property holds that subject name. */
@@ -117,11 +117,11 @@ struct LIVELINKHUBMESSAGING_API FLiveLinkHubTimecodeSettings
 	 * Number of frames to subtract from the qualified frame time when GetDelayedQualifiedFrameTime or GetDelayedTimecode is called.
 	 * @see GetDelayedQualifiedFrameTime, GetDelayedTimecode
 	 */
-	UPROPERTY(config, EditAnywhere, Category = "Timecode")
+	UPROPERTY(config, EditAnywhere, Category = "Timecode", meta = (ClampMin = "0", UIMin = "0", UIMax = "1200"))
 	float FrameDelay = 0.f;
 
 	/** The number of frame to keep in memory. The provider will not be synchronized until the buffer is full at least once. */
-	UPROPERTY(config, EditAnywhere, Category = "Timecode")
+	UPROPERTY(config, EditAnywhere, Category = "Timecode", meta = (ClampMin = "2", UIMin = "2", ClampMax = "10", UIMax = "10"))
 	int32 BufferSize;
 
 	/** How timecode should be evaluated. */
