Move LiveLinkHubSubjectSettings to a public header

#rnx

diff --git a/Engine/Plugins/Animation/LiveLinkHub/Source/LiveLinkHub/Private/LiveLinkHubModule.cpp b/Engine/Plugins/Animation/LiveLinkHub/Source/LiveLinkHub/Private/LiveLinkHubModule.cpp
--- a/Engine/Plugins/Animation/LiveLinkHub/Source/LiveLinkHub/Private/LiveLinkHubModule.cpp
+++ b/Engine/Plugins/Animation/LiveLinkHub/Source/LiveLinkHub/Private/LiveLinkHubModule.cpp
@@ -58,12 +58,8 @@ void FLiveLinkHubModule::StartupModule()
 	{
 		PropertyModule.RegisterCustomClassLayout(ULiveLinkSettings::StaticClass()->GetFName(), FOnGetDetailCustomizationInstance::CreateStatic(&FLiveLinkSettingsCustomization::MakeInstance));
 	}
-	bUseSubjectSettingsDetailsCustomization =
-		GConfig->GetBoolOrDefault(TEXT("LiveLink"), TEXT("bUseLiveLinkHubSubjectSettingsDetailsCustomization"), false, GEngineIni);
-	if (bUseSubjectSettingsDetailsCustomization)
-	{
-		PropertyModule.RegisterCustomClassLayout(ULiveLinkHubSubjectSettings::StaticClass()->GetFName(), FOnGetDetailCustomizationInstance::CreateStatic(&FLiveLinkHubSubjectSettingsDetailsCustomization::MakeInstance));
-	}
+	
+	PropertyModule.RegisterCustomClassLayout(ULiveLinkHubSubjectSettings::StaticClass()->GetFName(), FOnGetDetailCustomizationInstance::CreateStatic(&FLiveLinkHubSubjectSettingsDetailsCustomization::MakeInstance));
 }
 
 void FLiveLinkHubModule::ShutdownModule()
diff --git a/Engine/Plugins/Animation/LiveLinkHub/Source/LiveLinkHub/Private/LiveLinkHubSubjectSettings.h b/Engine/Plugins/Animation/LiveLinkHub/Source/LiveLinkHub/Private/LiveLinkHubSubjectSettings.h
deleted file mode 100644
diff --git a/Engine/Plugins/Animation/LiveLinkHub/Source/LiveLinkHub/Private/Subjects/LiveLinkHubSubjectSettingsDetailsCustomization.cpp b/Engine/Plugins/Animation/LiveLinkHub/Source/LiveLinkHub/Private/Subjects/LiveLinkHubSubjectSettingsDetailsCustomization.cpp
--- a/Engine/Plugins/Animation/LiveLinkHub/Source/LiveLinkHub/Private/Subjects/LiveLinkHubSubjectSettingsDetailsCustomization.cpp
+++ b/Engine/Plugins/Animation/LiveLinkHub/Source/LiveLinkHub/Private/Subjects/LiveLinkHubSubjectSettingsDetailsCustomization.cpp
@@ -6,6 +6,8 @@
 #include "DetailWidgetRow.h"
 #include "IDetailPropertyRow.h"
 #include "LiveLinkHubSubjectSettings.h"
+#include "Misc/ConfigCacheIni.h"
+
 
 #define LOCTEXT_NAMESPACE "LiveLinkHubSubjectSettingsDetailCustomization"
 
@@ -19,53 +21,64 @@ void FLiveLinkHubSubjectSettingsDetailsCustomization::CustomizeDetails(IDetailLa
 		return;
 	}
 
-	TSharedRef<IPropertyHandle> OutboundProperty = InDetailBuilder.GetProperty(GET_MEMBER_NAME_CHECKED(ULiveLinkHubSubjectSettings, OutboundName));
-	if (IDetailPropertyRow* PropertyRow = InDetailBuilder.EditDefaultProperty(OutboundProperty))
+	const bool bIsInLiveLinkHubApp = GConfig->GetBoolOrDefault(TEXT("LiveLink"), TEXT("bCreateLiveLinkHubInstance"), false, GEngineIni);
+
+	if (bIsInLiveLinkHubApp)
 	{
-		FResetToDefaultOverride ResetOverride;
+		TSharedRef<IPropertyHandle> OutboundProperty = InDetailBuilder.GetProperty(ULiveLinkHubSubjectSettings::GetOutboundNamePropertyName());
+		if (IDetailPropertyRow* PropertyRow = InDetailBuilder.EditDefaultProperty(OutboundProperty))
+		{
+			FResetToDefaultOverride ResetOverride;
+
+			const FName SubjectNamePropertyName = ULiveLinkHubSubjectSettings::GetSubjectNamePropertyName();
 
-		PropertyRow->CustomWidget()
-			.OverrideResetToDefault(FResetToDefaultOverride::Create(
-				FIsResetToDefaultVisible::CreateLambda([](TSharedPtr<IPropertyHandle> PropertyHandle)
-					{
-						TSharedPtr<IPropertyHandle> SubjectNameProperty = PropertyHandle->GetParentHandle()->GetChildHandle(GET_MEMBER_NAME_CHECKED(ULiveLinkHubSubjectSettings, SubjectName));
+			PropertyRow->CustomWidget()
+				.OverrideResetToDefault(FResetToDefaultOverride::Create(
+					FIsResetToDefaultVisible::CreateLambda([SubjectNamePropertyName](TSharedPtr<IPropertyHandle> PropertyHandle)
+						{
+							TSharedPtr<IPropertyHandle> SubjectNameProperty = PropertyHandle->GetParentHandle()->GetChildHandle(SubjectNamePropertyName);
 						
-						if (SubjectNameProperty && SubjectNameProperty->IsValidHandle())
+							if (SubjectNameProperty && SubjectNameProperty->IsValidHandle())
+							{
+								FText OutboundName;
+								FText SubjectName;
+								PropertyHandle->GetValueAsDisplayText(OutboundName);
+								SubjectNameProperty->GetValueAsDisplayText(SubjectName);
+								return !OutboundName.EqualTo(SubjectName);
+							}
+							return false;
+						}),
+					FResetToDefaultHandler::CreateLambda([SubjectNamePropertyName](TSharedPtr<IPropertyHandle> PropertyHandle)
 						{
-							FText OutboundName;
-							FText SubjectName;
-							PropertyHandle->GetValueAsDisplayText(OutboundName);
-							SubjectNameProperty->GetValueAsDisplayText(SubjectName);
-							return !OutboundName.EqualTo(SubjectName);
-						}
-						return false;
-					}),
-				FResetToDefaultHandler::CreateLambda([](TSharedPtr<IPropertyHandle> PropertyHandle)
-					{
-						TSharedPtr<IPropertyHandle> SubjectNameProperty = PropertyHandle->GetParentHandle()->GetChildHandle(GET_MEMBER_NAME_CHECKED(ULiveLinkHubSubjectSettings, SubjectName));
+							TSharedPtr<IPropertyHandle> SubjectNameProperty = PropertyHandle->GetParentHandle()->GetChildHandle(SubjectNamePropertyName);
 
-						if (SubjectNameProperty && SubjectNameProperty->IsValidHandle())
-						{
-							FString SubjectName;
-							SubjectNameProperty->GetValueAsDisplayString(SubjectName);
-							PropertyHandle->SetValue(*SubjectName);
-						}
-					})
-			))
-			.NameContent()
-			[
-				OutboundProperty->CreatePropertyNameWidget()
-			]
-			.ValueContent()
-			[
-				OutboundProperty->CreatePropertyValueWidget()
-			];
-	}
+							if (SubjectNameProperty && SubjectNameProperty->IsValidHandle())
+							{
+								FString SubjectName;
+								SubjectNameProperty->GetValueAsDisplayString(SubjectName);
+								PropertyHandle->SetValue(*SubjectName);
+							}
+						})
+				))
+				.NameContent()
+				[
+					OutboundProperty->CreatePropertyNameWidget()
+				]
+				.ValueContent()
+				[
+					OutboundProperty->CreatePropertyValueWidget()
+				];
+		}
 
-	InDetailBuilder.HideProperty(GET_MEMBER_NAME_CHECKED(ULiveLinkHubSubjectSettings, InterpolationProcessor), ULiveLinkSubjectSettings::StaticClass());
-	InDetailBuilder.HideProperty(GET_MEMBER_NAME_CHECKED(ULiveLinkHubSubjectSettings, FrameRate), ULiveLinkSubjectSettings::StaticClass());
-	InDetailBuilder.HideProperty(GET_MEMBER_NAME_CHECKED(ULiveLinkHubSubjectSettings, bRebroadcastSubject), ULiveLinkSubjectSettings::StaticClass());
-	InDetailBuilder.HideProperty(GET_MEMBER_NAME_CHECKED(ULiveLinkHubSubjectSettings, Translators), ULiveLinkSubjectSettings::StaticClass());
+		InDetailBuilder.HideProperty(GET_MEMBER_NAME_CHECKED(ULiveLinkHubSubjectSettings, InterpolationProcessor), ULiveLinkSubjectSettings::StaticClass());
+		InDetailBuilder.HideProperty(GET_MEMBER_NAME_CHECKED(ULiveLinkHubSubjectSettings, FrameRate), ULiveLinkSubjectSettings::StaticClass());
+		InDetailBuilder.HideProperty(GET_MEMBER_NAME_CHECKED(ULiveLinkHubSubjectSettings, bRebroadcastSubject), ULiveLinkSubjectSettings::StaticClass());
+		InDetailBuilder.HideProperty(GET_MEMBER_NAME_CHECKED(ULiveLinkHubSubjectSettings, Translators), ULiveLinkSubjectSettings::StaticClass());
+	}
+	else
+	{
+		InDetailBuilder.HideProperty(ULiveLinkHubSubjectSettings::GetOutboundNamePropertyName(), ULiveLinkHubSubjectSettings::StaticClass());
+	}
 }
 
 #undef LOCTEXT_NAMESPACE
diff --git a/Engine/Plugins/Animation/LiveLinkHub/Source/LiveLinkHub/Public/LiveLinkHubSubjectSettings.h b/Engine/Plugins/Animation/LiveLinkHub/Source/LiveLinkHub/Public/LiveLinkHubSubjectSettings.h
new file mode 100644
--- /dev/null
+++ b/Engine/Plugins/Animation/LiveLinkHub/Source/LiveLinkHub/Public/LiveLinkHubSubjectSettings.h
@@ -0,0 +1,64 @@
+// Copyright Epic Games, Inc. All Rights Reserved.
+
+#pragma once
+
+#include "LiveLinkSubjectSettings.h"
+
+#include "LiveLinkHubSubjectSettings.generated.h"
+
+class FLiveLinkHubClient;
+class ULiveLinkRole;
+
+/** Settings object for a livelinkhub subject. */
+UCLASS()
+class LIVELINKHUB_API ULiveLinkHubSubjectSettings : public ULiveLinkSubjectSettings
+{
+	GENERATED_BODY()
+public:
+
+	//~ Begin ULiveLinkSubjectSettings interface
+	virtual void Initialize(FLiveLinkSubjectKey InSubjectKey) override;
+
+	virtual FName GetRebroadcastName() const override
+	{
+		return *OutboundName;
+	}
+
+	//~ End ULiveLinkSubjectSettings interface
+
+	//~ Begin UObject interface
+	virtual void PreEditChange(FProperty* PropertyAboutToChange) override;
+	virtual void PostEditChangeProperty(FPropertyChangedEvent& PropertyChangedEvent) override;
+	//~ End UObject interface
+
+	static FName GetOutboundNamePropertyName()
+	{
+		return GET_MEMBER_NAME_CHECKED(ULiveLinkHubSubjectSettings, OutboundName);
+	}
+
+	static FName GetSubjectNamePropertyName()
+	{
+		return GET_MEMBER_NAME_CHECKED(ULiveLinkHubSubjectSettings, SubjectName);
+	}
+
+private:
+	/** Name of this subject. */
+	UPROPERTY(VisibleAnywhere, Category = "LiveLink")
+	FString SubjectName;
+
+	/** Name override that will be transmitted to clients instead of the subject name. */
+	UPROPERTY(EditAnywhere, Category = "LiveLink")
+	FString OutboundName;
+
+	/** Source that contains the subject. */
+	UPROPERTY(VisibleAnywhere, Category = "LiveLink")
+	FString Source;
+
+	/** Proxy property used edit the translators. */
+	UPROPERTY(EditAnywhere, Instanced, Category = "LiveLink", DisplayName = "Translator")
+	TObjectPtr<ULiveLinkFrameTranslator> TranslatorsProxy;
+
+private:
+	/* Previous outbound name to be used for noticing clients to remove this entry from their subject list. */
+	FName PreviousOutboundName;
+};
diff --git a/Engine/Source/Programs/LiveLinkHub/Config/DefaultEngine.ini b/Engine/Source/Programs/LiveLinkHub/Config/DefaultEngine.ini
--- a/Engine/Source/Programs/LiveLinkHub/Config/DefaultEngine.ini
+++ b/Engine/Source/Programs/LiveLinkHub/Config/DefaultEngine.ini
@@ -1,6 +1,5 @@
 [LiveLink]
 bUseLiveLinkHubSettingsDetailCustomization=true
-bUseLiveLinkHubSubjectSettingsDetailsCustomization=true
 bPanelControllerSeparateSourcesSubjects=true
 bUseModularClientReference=true
 bEnableParentSubjects=true
@@ -26,6 +25,7 @@ AutoConnectClients=LocalOnly
 +ProjectSettingsToDisplay=Live Link Hub
 +ProjectSettingsToDisplay=Timing & Sync
 +ProjectSettingsToDisplay=UDPMessaging
+bShowFrameRate=True
 
 [/Script/AppleARKit.AppleARKitSettings]
 bEnableLiveLinkForFaceTracking=true
@@ -64,3 +64,4 @@ NodeMemoryLoggingThresholdMB=5
 
 [CoreRedirects]
 +PropertyRedirects=(OldName="LiveLinkHubSettings.bAutoAddDiscoveredClients",NewName="AutoConnectClients")
+
