LiveLinkHub - Fix deleting clients only working with one client at a time

diff --git a/Engine/Plugins/Animation/LiveLinkHub/Source/LiveLinkHub/Private/Clients/LiveLinkHubClientsController.h b/Engine/Plugins/Animation/LiveLinkHub/Source/LiveLinkHub/Private/Clients/LiveLinkHubClientsController.h
--- a/Engine/Plugins/Animation/LiveLinkHub/Source/LiveLinkHub/Private/Clients/LiveLinkHubClientsController.h
+++ b/Engine/Plugins/Animation/LiveLinkHub/Source/LiveLinkHub/Private/Clients/LiveLinkHubClientsController.h
@@ -94,12 +94,11 @@ private:
 				// fallthrough
 			case ILiveLinkHubClientsModel::EClientEventType::Modified:
 			{
-				if (TOptional<FLiveLinkHubClientId> SelectedClient = ClientsView->GetSelectedClient())
+				TArray<FLiveLinkHubClientId> ClientIds = ClientsView->GetSelectedClients();
+
+				if (ClientIds.Num())
 				{
-					if (ClientIdentifier == *SelectedClient)
-					{
-						UpdateClientDetails(ClientIdentifier);
-					}
+					UpdateClientDetails(ClientIds[ClientIds.Num() - 1]);
 				}
 				break;
 			}
@@ -125,6 +124,8 @@ private:
 	{
 		if (TOptional<FLiveLinkHubUEClientInfo> ClientInfo = ClientsModel->GetClientInfo(Client))
 		{
+			DisplayedClientId = Client;
+
 			UEClientInfo->InitializeAs<FLiveLinkHubUEClientInfo>(*ClientInfo);
 			if (StructDetailsView)
 			{
@@ -155,11 +156,11 @@ private:
 		{
 			if (TSharedPtr<ILiveLinkHubSession> CurrentSession = SessionManager->GetCurrentSession())
 			{
-				if (TOptional<FLiveLinkHubClientId> SelectedClient = ClientsView->GetSelectedClient())
+				if (Client == DisplayedClientId)
 				{
-					// Clear client details if it was selected.
-					if (Client == *SelectedClient)
+					if (UEClientInfo)
 					{
+						// Clear client details if it was selected.
 						UEClientInfo->Reset();
 						StructDetailsView->SetStructureData(nullptr);
 					}
@@ -177,6 +178,8 @@ private:
 	TSharedPtr<ILiveLinkHubClientsModel> ClientsModel;
 	/** Holds a struct representing information about the currently selected client. */
 	TSharedPtr<TStructOnScope<FLiveLinkHubUEClientInfo>> UEClientInfo;
+	/** ID of the client displayed in the panel. */
+	FLiveLinkHubClientId DisplayedClientId;
 	/** Holds the details view created with the property editor module. */
 	TSharedPtr<IStructureDetailsView> StructDetailsView;
 	/** Delegate called when one of the connection's info has changed. */
diff --git a/Engine/Plugins/Animation/LiveLinkHub/Source/LiveLinkHub/Private/Clients/SLiveLinkHubClientsView.h b/Engine/Plugins/Animation/LiveLinkHub/Source/LiveLinkHub/Private/Clients/SLiveLinkHubClientsView.h
--- a/Engine/Plugins/Animation/LiveLinkHub/Source/LiveLinkHub/Private/Clients/SLiveLinkHubClientsView.h
+++ b/Engine/Plugins/Animation/LiveLinkHub/Source/LiveLinkHub/Private/Clients/SLiveLinkHubClientsView.h
@@ -485,18 +485,17 @@ public:
 		}
 	}
 
-	/** Get the currently selected client, or an invalid optional when nothing is selected. */
-	TOptional<FLiveLinkHubClientId> GetSelectedClient() const
+	/** Get the currently selected clients. */
+	TArray<FLiveLinkHubClientId> GetSelectedClients() const
 	{
-		TOptional<FLiveLinkHubClientId> ClientId;
+		TArray<FLiveLinkHubClientId> ClientIds;
 
 		TArray<FClientTreeItemPtr> SelectedClients = TreeView->GetSelectedItems();
-		if (SelectedClients.Num())
-		{
-			ClientId = SelectedClients[0]->ClientId;
-		}
+		Algo::TransformIf(SelectedClients, ClientIds, 
+			[](const FClientTreeItemPtr& ClientPtr) { return ClientPtr.IsValid(); },
+			[](const FClientTreeItemPtr& ClientPtr) { return ClientPtr->ClientId; });
 
-		return ClientId;
+		return ClientIds;
 	}
 
 private:
@@ -572,9 +571,9 @@ private:
     {
 		if (InKeyEvent.GetKey() == EKeys::Delete || InKeyEvent.GetKey() == EKeys::BackSpace)
     	{
-			if (TOptional<FLiveLinkHubClientId> Id = GetSelectedClient())
+			for (const FLiveLinkHubClientId& ClientId : GetSelectedClients())
 			{
-				OnRemoveClientFromSessionDelegate.ExecuteIfBound(*Id);
+				OnRemoveClientFromSessionDelegate.ExecuteIfBound(ClientId);
 			}
     		return FReply::Handled();
     	}
@@ -585,9 +584,9 @@ private:
 	/** Remove the selected client from the list. This will stop all livelink messages from being transmitted to it. */
 	void RemoveSelectedClient() const
 	{
-		if (TOptional<FLiveLinkHubClientId> Id = GetSelectedClient())
+		for (const FLiveLinkHubClientId& ClientId : GetSelectedClients())
 		{
-			OnRemoveClientFromSessionDelegate.ExecuteIfBound(*Id);
+			OnRemoveClientFromSessionDelegate.ExecuteIfBound(ClientId);
 		}
 	}
 
