Fix uninitialized property warnings in LiveLinkHubTimecodeSettings
Minor UI tweak to timecode framerate alignment

#rnx

diff --git a/Engine/Plugins/Animation/LiveLinkHub/Source/LiveLinkHub/Private/UI/Widgets/SLiveLinkTimecode.cpp b/Engine/Plugins/Animation/LiveLinkHub/Source/LiveLinkHub/Private/UI/Widgets/SLiveLinkTimecode.cpp
--- a/Engine/Plugins/Animation/LiveLinkHub/Source/LiveLinkHub/Private/UI/Widgets/SLiveLinkTimecode.cpp
+++ b/Engine/Plugins/Animation/LiveLinkHub/Source/LiveLinkHub/Private/UI/Widgets/SLiveLinkTimecode.cpp
@@ -97,7 +97,6 @@ void SLiveLinkTimecode::Construct(const FArguments& InArgs)
 		+ SHorizontalBox::Slot()
 		.AutoWidth()
 		.VAlign(VAlign_Center)
-		.MinWidth(75)
 		.Padding(FMargin(2, 0))
 		[
 			SNew(STimecode)
@@ -113,7 +112,7 @@ void SLiveLinkTimecode::Construct(const FArguments& InArgs)
 		+ SHorizontalBox::Slot()
 		.AutoWidth()
 		.VAlign(VAlign_Center)
-		.Padding(FMargin(0, 3, 0, 2))
+		.Padding(FMargin(2, 2, 0, 2))
 		[
 			SNew(STextBlock)
 			.Font(FCoreStyle::Get().GetFontStyle(TEXT("SmallText")))
diff --git a/Engine/Plugins/Animation/LiveLinkHub/Source/LiveLinkHubMessaging/Public/LiveLinkHubMessages.h b/Engine/Plugins/Animation/LiveLinkHub/Source/LiveLinkHubMessaging/Public/LiveLinkHubMessages.h
--- a/Engine/Plugins/Animation/LiveLinkHub/Source/LiveLinkHubMessaging/Public/LiveLinkHubMessages.h
+++ b/Engine/Plugins/Animation/LiveLinkHub/Source/LiveLinkHubMessaging/Public/LiveLinkHubMessages.h
@@ -122,11 +122,11 @@ struct LIVELINKHUBMESSAGING_API FLiveLinkHubTimecodeSettings
 
 	/** The number of frame to keep in memory. The provider will not be synchronized until the buffer is full at least once. */
 	UPROPERTY(config, EditAnywhere, Category = "Timecode", meta = (ClampMin = "2", UIMin = "2", ClampMax = "10", UIMax = "10"))
-	int32 BufferSize;
+	int32 BufferSize = 2;
 
 	/** How timecode should be evaluated. */
 	UPROPERTY(config, EditAnywhere, Category = "Timecode")
-	ELiveLinkTimecodeProviderEvaluationType EvaluationType;
+	ELiveLinkTimecodeProviderEvaluationType EvaluationType = ELiveLinkTimecodeProviderEvaluationType::Lerp;
 
 	/** Assign the settings to a new timecode provider and override the current engine settings. */
 	void AssignTimecodeSettingsAsProviderToEngine() const;
