LiveLinkHub - Automatically enter rename when duplicating a recording

diff --git a/Engine/Plugins/Animation/LiveLinkHub/Source/LiveLinkHub/Private/Recording/SLiveLinkHubRecordingListView.h b/Engine/Plugins/Animation/LiveLinkHub/Source/LiveLinkHub/Private/Recording/SLiveLinkHubRecordingListView.h
--- a/Engine/Plugins/Animation/LiveLinkHub/Source/LiveLinkHub/Private/Recording/SLiveLinkHubRecordingListView.h
+++ b/Engine/Plugins/Animation/LiveLinkHub/Source/LiveLinkHub/Private/Recording/SLiveLinkHubRecordingListView.h
@@ -69,7 +69,7 @@ public:
 			+ SVerticalBox::Slot()
 			.VAlign(VAlign_Fill)
 			[
-				SNew(SBox)
+				SAssignNew(BoxWidget, SBox)
 				.Visibility(this, &SLiveLinkHubRecordingListView::GetRecordingPickerVisibility)
 				[
 					CreateRecordingPicker()
@@ -140,7 +140,7 @@ private:
 	}
 
 	/** Creates the asset picker widget for selecting a recording. */
-	TSharedRef<SWidget> CreateRecordingPicker()
+	TSharedRef<SWidget> CreateRecordingPicker(TOptional<FAssetData> AssetData = {})
 	{
 		FMenuBuilder MenuBuilder(true, nullptr);
 
@@ -166,6 +166,11 @@ private:
 			AssetPickerConfig.Filter = MakeAssetFilter();
 			AssetPickerConfig.OnAssetDoubleClicked = FOnAssetSelected::CreateRaw(this, &SLiveLinkHubRecordingListView::OnImportRecording);
 			AssetPickerConfig.OnGetAssetContextMenu = FOnGetAssetContextMenu::CreateRaw(this, &SLiveLinkHubRecordingListView::GetAssetContextMenu);
+
+			if (AssetData.IsSet())
+			{
+				AssetPickerConfig.InitialAssetSelection = *AssetData;
+			}
 		}
 
 		MenuBuilder.BeginSection(NAME_None, LOCTEXT("ImportRecording_MenuSection", "Import Recording"));
@@ -197,7 +202,7 @@ private:
 		return Filter;
 	}
 
-	TSharedPtr<SWidget> GetAssetContextMenu(const TArray<FAssetData>& SelectedAssets) const
+	TSharedPtr<SWidget> GetAssetContextMenu(const TArray<FAssetData>& SelectedAssets)
 	{
 		if (SelectedAssets.Num() <= 0)
 		{
@@ -236,7 +241,7 @@ private:
 				LOCTEXT("DuplicateRecordingTooltip", "Duplicate the recording"),
 				FSlateIcon(FAppStyle::GetAppStyleSetName(), "GenericCommands.Duplicate"),
 				FUIAction(
-					FExecuteAction::CreateLambda([SelectedAsset] ()
+					FExecuteAction::CreateLambda([this, SelectedAsset] ()
 					{
 						if (SelectedAsset.IsValid())
 						{
@@ -263,6 +268,23 @@ private:
 							{
 								CastChecked<ULiveLinkUAssetRecording>(SelectedAsset)->UnloadRecordingData();
 							}
+
+							// There is no inherent way to update the selection of the asset picker, so instead we'll recreate one that is already selecting the new asset.
+							BoxWidget->SetContent(CreateRecordingPicker(FAssetData{ NewAsset }));
+
+							// It may take a few frames for the selection to fully update in the new picker, so give it ample time to do so before triggering the rename.
+							GEditor->GetTimerManager()->SetTimer(TimerHandle, [this]() 
+							{
+								if (TimerHandle.IsValid())
+								{
+									if (const FContentBrowserModule* ContentBrowserModule = FModuleManager::Get().GetModulePtr<FContentBrowserModule>("ContentBrowser"))
+									{
+										ContentBrowserModule->Get().ExecuteRename(AssetPicker);
+									}
+								}
+							}, 0.3, false);
+
+							
 						}
 					}),
 					FCanExecuteAction::CreateLambda([] () { return true; })
@@ -306,11 +328,14 @@ private:
 	FOnImportRecording OnImportRecordingDelegate;
 	/** The asset picker used for selecting recordings. */
 	TSharedPtr<SWidget> AssetPicker;
-	
 	/** Handle for when an asset is added to the asset registry. */
 	FDelegateHandle OnAssetAddedHandle;
 	/** Handle for when an asset is removed from the asset registry. */
 	FDelegateHandle OnAssetRemovedHandle;
+	/** Box widget used to hold the asset picker. */
+	TSharedPtr<SBox> BoxWidget;
+	/** Timer handle used for triggering a rename after duplicating a recording. */
+	FTimerHandle TimerHandle;
 	/** True if there are recording assets that exist. */
 	mutable TOptional<bool> bAssetsAvailableCached;
 };
