LiveLink - Fix broken UI in LL client & crash when using LiveLinkHubSubjectSettings in editor

diff --git a/Engine/Plugins/Animation/LiveLink/Source/LiveLinkEditor/Private/LiveLinkClientPanel.cpp b/Engine/Plugins/Animation/LiveLink/Source/LiveLinkEditor/Private/LiveLinkClientPanel.cpp
--- a/Engine/Plugins/Animation/LiveLink/Source/LiveLinkEditor/Private/LiveLinkClientPanel.cpp
+++ b/Engine/Plugins/Animation/LiveLink/Source/LiveLinkEditor/Private/LiveLinkClientPanel.cpp
@@ -136,20 +136,22 @@ void SLiveLinkClientPanel::Construct(const FArguments& Args, FLiveLinkClient* In
 					.HAlign(HAlign_Left)
 					.VAlign(VAlign_Bottom)
 					[
-						SNew(SHorizontalBox)
-						+SHorizontalBox::Slot()
-						.AutoWidth()
-						.Padding(10, 4, 4, 10)
+						SNew(SBox)
+						.Padding(FMargin(100, 4, 4, 10))
 						[
-							SNew(STextBlock)
-							.Text(this, &SLiveLinkClientPanel::GetMessageCountText)
-						]
-						+SHorizontalBox::Slot()
-						.Padding(20, 4, 50, 10)
-						.AutoWidth()
-						[
-							SNew(STextBlock)
-							.Text(this, &SLiveLinkClientPanel::GetSelectedMessageOccurrenceText)
+							SNew(SHorizontalBox)
+							+ SHorizontalBox::Slot()
+							.AutoWidth()
+							[
+								SNew(STextBlock)
+								.Text(this, &SLiveLinkClientPanel::GetMessageCountText)
+							]
+							+ SHorizontalBox::Slot()
+							.AutoWidth()
+							[
+								SNew(STextBlock)
+								.Text(this, &SLiveLinkClientPanel::GetSelectedMessageOccurrenceText)
+							]
 						]
 					]
 				]
diff --git a/Engine/Plugins/Animation/LiveLinkHub/Source/LiveLinkHub/Private/LiveLinkHubSubjectSettings.cpp b/Engine/Plugins/Animation/LiveLinkHub/Source/LiveLinkHub/Private/LiveLinkHubSubjectSettings.cpp
--- a/Engine/Plugins/Animation/LiveLinkHub/Source/LiveLinkHub/Private/LiveLinkHubSubjectSettings.cpp
+++ b/Engine/Plugins/Animation/LiveLinkHub/Source/LiveLinkHub/Private/LiveLinkHubSubjectSettings.cpp
@@ -2,11 +2,12 @@
 
 #include "LiveLinkHubSubjectSettings.h"
 
+#include "Clients/LiveLinkHubProvider.h"
 #include "LiveLinkHubClient.h"
 #include "LiveLinkHubModule.h"
 #include "LiveLinkHubSubjectSettingsUtils.h"
 #include "Features/IModularFeatures.h"
-#include "Clients/LiveLinkHubProvider.h"
+#include "Misc/ConfigCacheIni.h"
 
 void ULiveLinkHubSubjectSettings::Initialize(FLiveLinkSubjectKey InSubjectKey)
 {
@@ -36,6 +37,14 @@ void ULiveLinkHubSubjectSettings::PostEditChangeProperty(FPropertyChangedEvent&
 {
 	Super::PostEditChangeProperty(PropertyChangedEvent);
 
+	const bool bIsInLiveLinkHubApp = GConfig->GetBoolOrDefault(TEXT("LiveLink"), TEXT("bCreateLiveLinkHubInstance"), false, GEngineIni);
+
+	if (!bIsInLiveLinkHubApp)
+	{
+		// We're not using most of these settings in editor, so there's no use in calling CacheSubjectSettings or modifying the OutboundName.
+		return;
+	}
+
 	if (PropertyChangedEvent.GetPropertyName() == GET_MEMBER_NAME_CHECKED(ULiveLinkHubSubjectSettings, OutboundName))
 	{
 		if (PreviousOutboundName != *OutboundName)
