Live Link Hub: Load session via `-SessionPath="D:\blah.json"` command line argument.

Also implements corresponding support for the argument in the Switchboard device.

diff --git a/Engine/Plugins/Animation/LiveLinkHub/Source/LiveLinkHub/Private/LiveLinkHub.cpp b/Engine/Plugins/Animation/LiveLinkHub/Source/LiveLinkHub/Private/LiveLinkHub.cpp
--- a/Engine/Plugins/Animation/LiveLinkHub/Source/LiveLinkHub/Private/LiveLinkHub.cpp
+++ b/Engine/Plugins/Animation/LiveLinkHub/Source/LiveLinkHub/Private/LiveLinkHub.cpp
@@ -194,6 +194,12 @@ void FLiveLinkHub::InitializePostEngineInit()
 	GetDefault<ULiveLinkHubTimeAndSyncSettings>()->ApplyCustomTimeStep();
 
 	GIsRunning = true;
+
+	FString SessionPath;
+	if (FParse::Value(FCommandLine::Get(), TEXT("SessionPath="), SessionPath) && SessionPath.Len())
+	{
+		SessionManager->RestoreSession(SessionPath);
+	}
 }
 
 FLiveLinkHub::~FLiveLinkHub()
diff --git a/Engine/Plugins/Animation/LiveLinkHub/Source/LiveLinkHub/Private/Session/LiveLinkHubSessionManager.h b/Engine/Plugins/Animation/LiveLinkHub/Source/LiveLinkHub/Private/Session/LiveLinkHubSessionManager.h
--- a/Engine/Plugins/Animation/LiveLinkHub/Source/LiveLinkHub/Private/Session/LiveLinkHubSessionManager.h
+++ b/Engine/Plugins/Animation/LiveLinkHub/Source/LiveLinkHub/Private/Session/LiveLinkHubSessionManager.h
@@ -45,8 +45,8 @@ public:
 	/** Prompt the user to save the current session in a given directory. */
 	virtual void SaveSessionAs() = 0;
 
-	/** Prompt the user to pick a session file to restore. */
-	virtual void RestoreSession() = 0;
+	/** Restore a session from file. If not provided, will prompt the user to browse for the file. */
+	virtual void RestoreSession(FStringView InSessionPath = {}) = 0;
 
 	/** Save the current session. If not path is specified, the last save path will be used. */
 	virtual void SaveCurrentSession(const FString& SavePath = TEXT("")) = 0;
@@ -202,35 +202,43 @@ public:
 		UE::LiveLinkHub::FileUtilities::Private::SaveConfig(LiveLinkHubSessionData, LastConfigPath);
 	}
 
-	virtual void RestoreSession() override
+	virtual void RestoreSession(FStringView InSessionPath /* = {} */) override
 	{
-		const FString FileDescription = UE::LiveLinkHub::FileUtilities::Private::ConfigDescription;
-		const FString Extensions = UE::LiveLinkHub::FileUtilities::Private::ConfigExtension;
-		const FString FileTypes = FString::Printf(TEXT("%s (*.%s)|*.%s"), *FileDescription, *Extensions, *Extensions);
-
-		const FString DefaultFile = UE::LiveLinkHub::FileUtilities::Private::ConfigDefaultFileName;
-
+		// If no path was provided, prompt the user to browse for a session file.
 		TArray<FString> OpenFileNames;
+		if (InSessionPath.IsEmpty())
+		{
+			const FString FileDescription = UE::LiveLinkHub::FileUtilities::Private::ConfigDescription;
+			const FString Extensions = UE::LiveLinkHub::FileUtilities::Private::ConfigExtension;
+			const FString FileTypes = FString::Printf(TEXT("%s (*.%s)|*.%s"), *FileDescription, *Extensions, *Extensions);
+
+			const FString DefaultFile = UE::LiveLinkHub::FileUtilities::Private::ConfigDefaultFileName;
+
+			IDesktopPlatform* DesktopPlatform = FDesktopPlatformModule::Get();
+			const void* ParentWindowWindowHandle = FSlateApplication::Get().FindBestParentWindowHandleForDialogs(nullptr);
+			OpenSessionTimestampsSeconds = FPlatformTime::Seconds();
+			const bool bFileSelected = DesktopPlatform->OpenFileDialog(
+				ParentWindowWindowHandle,
+				LOCTEXT("LiveLinkHubOpenTitle", "Open").ToString(),
+				FEditorDirectories::Get().GetLastDirectory(ELastDirectory::GENERIC_OPEN),
+				DefaultFile,
+				FileTypes,
+				EFileDialogFlags::None,
+				OpenFileNames);
+
+			if (bFileSelected && OpenFileNames.Num() > 0)
+			{
+				InSessionPath = OpenFileNames[0];
+			}
+		}
 
-		IDesktopPlatform* DesktopPlatform = FDesktopPlatformModule::Get();
-		const void* ParentWindowWindowHandle = FSlateApplication::Get().FindBestParentWindowHandleForDialogs(nullptr);
-		OpenSessionTimestampsSeconds = FPlatformTime::Seconds();
-		const bool bFileSelected = DesktopPlatform->OpenFileDialog(
-			ParentWindowWindowHandle,
-			LOCTEXT("LiveLinkHubOpenTitle", "Open").ToString(),
-			FEditorDirectories::Get().GetLastDirectory(ELastDirectory::GENERIC_OPEN),
-			DefaultFile,
-			FileTypes,
-			EFileDialogFlags::None,
-			OpenFileNames);
-
-		if (bFileSelected && OpenFileNames.Num() > 0)
+		if (!InSessionPath.IsEmpty())
 		{
 			// Certain sources may take time to clean up. If they don't complete in time then the new config being loaded may not create
 			// duplicate sources correctly. There should be errors in the logs of the sources that failed to remove or were unable to be added.
 			constexpr bool bWaitForSourceRemoval = true;
 			ClearSession(bWaitForSourceRemoval);
-			RestoreSession(OpenFileNames[0]);
+			InternalRestoreSession(InSessionPath);
 		}
 	}
 
@@ -359,7 +367,7 @@ public:
 
 private:
 	/** Load a session from disk and restore its content. */
-	void RestoreSession(const FString& Path)
+	void InternalRestoreSession(FStringView Path)
 	{
 		LastConfigPath = Path;
 		FEditorDirectories::Get().SetLastDirectory(ELastDirectory::GENERIC_OPEN, FPaths::GetPath(LastConfigPath));
diff --git a/Engine/Plugins/VirtualProduction/Switchboard/Source/Switchboard/switchboard/devices/livelinkhub/plugin_livelinkhub.py b/Engine/Plugins/VirtualProduction/Switchboard/Source/Switchboard/switchboard/devices/livelinkhub/plugin_livelinkhub.py
--- a/Engine/Plugins/VirtualProduction/Switchboard/Source/Switchboard/switchboard/devices/livelinkhub/plugin_livelinkhub.py
+++ b/Engine/Plugins/VirtualProduction/Switchboard/Source/Switchboard/switchboard/devices/livelinkhub/plugin_livelinkhub.py
@@ -8,7 +8,7 @@ from typing import Callable
 from PySide6 import QtCore, QtWidgets
 
 from switchboard import message_protocol, switchboard_widgets
-from switchboard.config import CONFIG, Setting, BoolSetting, StringSetting
+from switchboard.config import CONFIG, Setting, BoolSetting, FilePathSetting, StringSetting
 from switchboard.devices.device_base import Device, DeviceStatus, PluginHeaderWidgets
 from switchboard.devices.unreal.plugin_unreal import DeviceUnreal, DeviceWidgetUnreal, ProgramStartQueueItem
 from switchboard.switchboard_logging import LOGGER
@@ -40,6 +40,13 @@ class DeviceLiveLinkHub(DeviceUnreal):
             tool_tip='Additional command line arguments for Live Link Hub',
             category='General Settings',
         ),
+        'session_path': FilePathSetting(
+            attr_name='session_path',
+            nice_name='Session File Path',
+            value='',
+            tool_tip='Path to the Live Link Hub session file to load on startup',
+            category='General Settings',
+        ),
         'retrieve_logs': BoolSetting(
             attr_name='retrieve_logs',
             nice_name='Retrieve Logs',
@@ -132,13 +139,17 @@ class DeviceLiveLinkHub(DeviceUnreal):
     #@override
     @property
     def executable_filename(self):
-        return self.csettings['executable_name'].get_value()
+        return DeviceLiveLinkHub.csettings['executable_name'].get_value()
 
     #@override
     @property
     def extra_cmdline_args_setting(self):
         return DeviceLiveLinkHub.csettings['command_line_arguments'].get_value(self.name)
 
+    @property
+    def session_path(self):
+        return DeviceLiveLinkHub.csettings['session_path'].get_value(self.name)
+
     #@override
     def get_remote_log_path(self):
         return self.llh_proj_dir / 'Saved' / 'Logs'
@@ -151,7 +162,8 @@ class DeviceLiveLinkHub(DeviceUnreal):
 
         command_line_args += ' -hub' if self.is_primary() else ' -spoke'
 
-        # TODO: Specify config to load. Need to implement on LLH side first
+        if self.session_path:
+            command_line_args += f' -SessionPath="{self.session_path}"'
 
         return command_line_args
 
