LiveLinkHub - Fix timecode and custom time step used on startup regardless of their respective "enable" setting.

#rnx

diff --git a/Engine/Plugins/Animation/LiveLinkHub/Source/LiveLinkHub/Private/LiveLinkHub.cpp b/Engine/Plugins/Animation/LiveLinkHub/Source/LiveLinkHub/Private/LiveLinkHub.cpp
--- a/Engine/Plugins/Animation/LiveLinkHub/Source/LiveLinkHub/Private/LiveLinkHub.cpp
+++ b/Engine/Plugins/Animation/LiveLinkHub/Source/LiveLinkHub/Private/LiveLinkHub.cpp
@@ -190,8 +190,17 @@ void FLiveLinkHub::InitializePostEngineInit()
 
 	PlaybackController->Start();
 
-	GetDefault<ULiveLinkHubTimeAndSyncSettings>()->ApplyTimecodeProvider();
-	GetDefault<ULiveLinkHubTimeAndSyncSettings>()->ApplyCustomTimeStep();
+	const ULiveLinkHubTimeAndSyncSettings* TimeAndSyncSettings = GetDefault<ULiveLinkHubTimeAndSyncSettings>();
+
+	if (TimeAndSyncSettings->bUseLiveLinkHubAsTimecodeSource)
+	{
+		TimeAndSyncSettings->ApplyTimecodeProvider();
+	}
+
+	if (TimeAndSyncSettings->bUseLiveLinkHubAsCustomTimeStepSource)
+	{
+		TimeAndSyncSettings->ApplyCustomTimeStep();
+	}
 
 	GIsRunning = true;
 
