LiveLinkHub - Allow picking a startup config

diff --git a/Engine/Plugins/Animation/LiveLinkHub/Source/LiveLinkHub/Private/LiveLinkHub.cpp b/Engine/Plugins/Animation/LiveLinkHub/Source/LiveLinkHub/Private/LiveLinkHub.cpp
--- a/Engine/Plugins/Animation/LiveLinkHub/Source/LiveLinkHub/Private/LiveLinkHub.cpp
+++ b/Engine/Plugins/Animation/LiveLinkHub/Source/LiveLinkHub/Private/LiveLinkHub.cpp
@@ -209,6 +209,10 @@ void FLiveLinkHub::InitializePostEngineInit()
 	{
 		SessionManager->RestoreSession(SessionPath);
 	}
+	else if (!GetDefault<ULiveLinkHubSettings>()->StartupConfig.FilePath.IsEmpty())
+	{ 
+		SessionManager->RestoreSession(GetDefault<ULiveLinkHubSettings>()->StartupConfig.FilePath);
+	}
 }
 
 FLiveLinkHub::~FLiveLinkHub()
diff --git a/Engine/Plugins/Animation/LiveLinkHub/Source/LiveLinkHub/Private/Session/LiveLinkHubSessionManager.h b/Engine/Plugins/Animation/LiveLinkHub/Source/LiveLinkHub/Private/Session/LiveLinkHubSessionManager.h
--- a/Engine/Plugins/Animation/LiveLinkHub/Source/LiveLinkHub/Private/Session/LiveLinkHubSessionManager.h
+++ b/Engine/Plugins/Animation/LiveLinkHub/Source/LiveLinkHub/Private/Session/LiveLinkHubSessionManager.h
@@ -196,8 +196,9 @@ public:
 		if (!SavePath.IsEmpty())
 		{
 			LastConfigPath = SavePath;
+			GetMutableDefault<ULiveLinkHubUserSettings>()->LastConfigDirectory = FPaths::GetPath(LastConfigPath);
+			GetMutableDefault<ULiveLinkHubUserSettings>()->SaveConfig();
 		}
-		FEditorDirectories::Get().SetLastDirectory(ELastDirectory::GENERIC_SAVE, FPaths::GetPath(LastConfigPath));
 
 		UE::LiveLinkHub::FileUtilities::Private::SaveConfig(LiveLinkHubSessionData, LastConfigPath);
 	}
@@ -214,13 +215,19 @@ public:
 
 			const FString DefaultFile = UE::LiveLinkHub::FileUtilities::Private::ConfigDefaultFileName;
 
+			FString DefaultOpenPath = FEditorDirectories::Get().GetLastDirectory(ELastDirectory::GENERIC_OPEN);
+			if (!GetDefault<ULiveLinkHubUserSettings>()->LastConfigDirectory.IsEmpty())
+			{
+				DefaultOpenPath = GetDefault<ULiveLinkHubUserSettings>()->LastConfigDirectory;
+			}
+
 			IDesktopPlatform* DesktopPlatform = FDesktopPlatformModule::Get();
 			const void* ParentWindowWindowHandle = FSlateApplication::Get().FindBestParentWindowHandleForDialogs(nullptr);
 			OpenSessionTimestampsSeconds = FPlatformTime::Seconds();
 			const bool bFileSelected = DesktopPlatform->OpenFileDialog(
 				ParentWindowWindowHandle,
 				LOCTEXT("LiveLinkHubOpenTitle", "Open").ToString(),
-				FEditorDirectories::Get().GetLastDirectory(ELastDirectory::GENERIC_OPEN),
+				DefaultOpenPath,
 				DefaultFile,
 				FileTypes,
 				EFileDialogFlags::None,
@@ -229,6 +236,8 @@ public:
 			if (bFileSelected && OpenFileNames.Num() > 0)
 			{
 				InSessionPath = OpenFileNames[0];
+				GetMutableDefault<ULiveLinkHubUserSettings>()->LastConfigDirectory = FPaths::GetPath(OpenFileNames[0]);
+				GetMutableDefault<ULiveLinkHubUserSettings>()->SaveConfig();
 			}
 		}
 
@@ -370,7 +379,6 @@ private:
 	void InternalRestoreSession(FStringView Path)
 	{
 		LastConfigPath = Path;
-		FEditorDirectories::Get().SetLastDirectory(ELastDirectory::GENERIC_OPEN, FPaths::GetPath(LastConfigPath));
 
 		if (ULiveLinkHubSessionData* SessionData = UE::LiveLinkHub::FileUtilities::Private::LoadConfig(LastConfigPath))
 		{
diff --git a/Engine/Plugins/Animation/LiveLinkHub/Source/LiveLinkHub/Private/Settings/LiveLinkHubSettings.h b/Engine/Plugins/Animation/LiveLinkHub/Source/LiveLinkHub/Private/Settings/LiveLinkHubSettings.h
--- a/Engine/Plugins/Animation/LiveLinkHub/Source/LiveLinkHub/Private/Settings/LiveLinkHubSettings.h
+++ b/Engine/Plugins/Animation/LiveLinkHub/Source/LiveLinkHub/Private/Settings/LiveLinkHubSettings.h
@@ -35,6 +35,10 @@ protected:
 	//~ End UObject interface
 	
 public:
+	/** Config to apply when starting LiveLinkHub. */
+	UPROPERTY(config, EditAnywhere, Category = "LiveLinkHub")
+	FFilePath StartupConfig;
+
 	/** If enabled, discovered clients will be automatically added to the current session. */
 	UPROPERTY(config, EditAnywhere, Category="LiveLinkHub", DisplayName = "Auto Connect Mode")
 	ELiveLinkHubAutoConnectMode AutoConnectClients = ELiveLinkHubAutoConnectMode::LocalOnly;
@@ -120,4 +124,8 @@ public:
 	/** Which directories to scan to discover layouts */
 	UPROPERTY(config, EditAnywhere, Category = "LiveLinkHub")
 	TArray<FString> LayoutDirectories;
+
+	/** The last directory of the a config that was saved or loaded. */
+	UPROPERTY(config)
+	FString LastConfigDirectory;
 };
