LiveLinkHub - Allow re-connecting to previous instances

diff --git a/Engine/Plugins/Animation/LiveLinkHub/Source/LiveLinkHubMessaging/Private/LiveLinkHubConnectionManager.h b/Engine/Plugins/Animation/LiveLinkHub/Source/LiveLinkHubMessaging/Private/LiveLinkHubConnectionManager.h
--- a/Engine/Plugins/Animation/LiveLinkHub/Source/LiveLinkHubMessaging/Private/LiveLinkHubConnectionManager.h
+++ b/Engine/Plugins/Animation/LiveLinkHub/Source/LiveLinkHubMessaging/Private/LiveLinkHubConnectionManager.h
@@ -15,6 +15,7 @@
 #include "LiveLinkMessageBusDiscoveryManager.h"
 #include "LiveLinkMessageBusFinder.h"
 #include "LiveLinkSettings.h"
+#include "Misc/ConfigCacheIni.h"
 #include "Misc/CoreDelegates.h"
 #include "Modules/ModuleManager.h"
 #include "TimerManager.h"
@@ -144,6 +145,8 @@ public:
 	{
 		FCoreUObjectDelegates::PostLoadMapWithWorld.AddRaw(this, &FLiveLinkHubConnectionManager::PostLoadMap);
 		FCoreDelegates::OnPostEngineInit.AddRaw(this, &FLiveLinkHubConnectionManager::StartDiscovery);
+
+		bEnableReconnectingToStaleSource = GConfig->GetBoolOrDefault(TEXT("LiveLink"), TEXT("bEnableReconnectingToStaleSource"), true, GEngineIni);
 	}
 
 	~FLiveLinkHubConnectionManager()
@@ -230,9 +233,17 @@ private:
 			{
 				if (LiveLinkClient->GetSourceType(SourceId).ToString() == PollResult->Name && LiveLinkClient->GetSourceMachineName(SourceId).ToString() == PollResult->MachineName)
 				{
-					// Todo: Should we re-enable stale connections? 
-					UE_LOG(LogLiveLinkHubConnectionManager, Verbose, TEXT("Rejecting poll result since source %s already exists."), *PollResult->Name);
-					return;
+					// If we're reconnecting to an invalid source, make sure to delete the previous one first. 
+					if (bEnableReconnectingToStaleSource && !LiveLinkClient->GetSourceStatus(SourceId).EqualToCaseIgnored(FLiveLinkMessageBusSource::ValidSourceStatus()))
+					{
+						// todo? We may want to eventually keep the source but "forwarding" the connection string to the source in order to keep the previous source settings.
+						LiveLinkClient->RemoveSource(SourceId);
+					}
+					else
+					{
+						UE_LOG(LogLiveLinkHubConnectionManager, Verbose, TEXT("Rejecting poll result since source %s already exists."), *PollResult->Name);
+						return;
+					}
 				}
 			}
 
@@ -262,6 +273,9 @@ private:
 	FOnGetTopologyMode GetTopologyModeDelegate;
 	/** Get the instance id (Only relevant if this is running inside of livelinkhub. */
 	FOnGetInstanceId GetInstanceIdDelegate;
+	/** Whether to allow reconnecting to stale LLH sources. */
+	bool bEnableReconnectingToStaleSource = true;
+	
 };
 #else
 class FLiveLinkHubConnectionManager
diff --git a/Engine/Source/Programs/LiveLinkHub/Config/DefaultEngine.ini b/Engine/Source/Programs/LiveLinkHub/Config/DefaultEngine.ini
--- a/Engine/Source/Programs/LiveLinkHub/Config/DefaultEngine.ini
+++ b/Engine/Source/Programs/LiveLinkHub/Config/DefaultEngine.ini
@@ -4,6 +4,7 @@ bPanelControllerSeparateSourcesSubjects=true
 bUseModularClientReference=true
 bEnableParentSubjects=true
 bCreateLiveLinkHubInstance=true
+bEnableReconnectingToStaleSource=true
 
 [LiveLinkHub]
 bEnableLLHMessageBusSourceFactory=true
