LiveLinkHub - Add icon for genlock and color it depending on its validity.
Move Hub/Spoke mode selection to status bar
Add icons for Hub/Spoke modes
Move Recording info closer to  Record button
Rename Creator mode to Live Data

diff --git a/Engine/Plugins/Animation/LiveLink/Source/LiveLinkEditor/Private/LiveLinkEditorModule.cpp b/Engine/Plugins/Animation/LiveLink/Source/LiveLinkEditor/Private/LiveLinkEditorModule.cpp
--- a/Engine/Plugins/Animation/LiveLink/Source/LiveLinkEditor/Private/LiveLinkEditorModule.cpp
+++ b/Engine/Plugins/Animation/LiveLink/Source/LiveLinkEditor/Private/LiveLinkEditorModule.cpp
@@ -117,6 +117,9 @@ public:
 		StyleSet->Set("LiveLinkHub.Icon.Small", new IMAGE_PLUGIN_BRUSH(TEXT("Starship/LiveLinkHub_16x16"), Icon16x16));
 		StyleSet->Set("LiveLinkHub.Icon", new IMAGE_PLUGIN_BRUSH_SVG(TEXT("Starship/LiveLinkHub_Icon"), Icon40x40));
 		StyleSet->Set("LiveLinkHub.EllipsisIcon", new IMAGE_PLUGIN_BRUSH_SVG(TEXT("Starship/Ellipsis"), Icon16x16));
+		StyleSet->Set("LiveLinkHub.TimecodeGenlock", new IMAGE_PLUGIN_BRUSH_SVG(TEXT("Starship/TimecodeGenlock"), Icon16x16));
+		StyleSet->Set("LiveLinkHub.SpokeMode", new IMAGE_PLUGIN_BRUSH_SVG(TEXT("Starship/SpokeMode"), Icon16x16));
+		StyleSet->Set("LiveLinkHub.HubMode", new IMAGE_PLUGIN_BRUSH_SVG(TEXT("Starship/HubMode"), Icon16x16));
 
 
 		StyleSet->Set("ClassIcon.LiveLinkPreset",                      new IMAGE_PLUGIN_BRUSH_SVG("Starship/LiveLink", Icon16x16));
diff --git a/Engine/Plugins/Animation/LiveLinkHub/Source/LiveLinkHub/Private/LiveLinkHubApplicationMode.cpp b/Engine/Plugins/Animation/LiveLinkHub/Source/LiveLinkHub/Private/LiveLinkHubApplicationMode.cpp
--- a/Engine/Plugins/Animation/LiveLinkHub/Source/LiveLinkHub/Private/LiveLinkHubApplicationMode.cpp
+++ b/Engine/Plugins/Animation/LiveLinkHub/Source/LiveLinkHub/Private/LiveLinkHubApplicationMode.cpp
@@ -6,6 +6,7 @@
 #include "UI/Widgets/SLiveLinkHubSettings.h"
 #include "WorkflowOrientedApp/WorkflowCentricApplication.h"
 #include "ToolMenus.h"
+#include "UI/Widgets/SLiveLinkHubTopologyModeSwitcher.h"
 
 const FName ILiveLinkHubApplicationModeFactory::ModularFeatureName = "LiveLinkHubApplicationModeFactory";
 const FName FLiveLinkHubApplicationMode::FileMenuExtensionPoint = "LiveLinkHubFileMenu";
@@ -42,17 +43,45 @@ void FLiveLinkHubApplicationMode::PostActivateMode()
 {
 	FToolMenuOwnerScoped OwnerScoped(GetModeName());
 
-	TArray<TSharedRef<SWidget>> Widgets = GetStatusBarWidgets_Impl();
 
 	UToolMenu* Menu = UToolMenus::Get()->ExtendMenu(TEXT("LiveLinkHub.ToolBar"));
-	FToolMenuSection& SessionConfigSection = Menu->AddSection(TEXT("ModeWidgets"), FText::GetEmpty(), FToolMenuInsert(NAME_None, EToolMenuInsertType::First));
 
+	FToolMenuSection& GlobalWidgetsSection = Menu->AddSection(TEXT("GlobalWidgets"), FText::GetEmpty(), FToolMenuInsert(NAME_None, EToolMenuInsertType::First));
+
+	TSharedRef<SWidget> GlobalWidget =
+		SNew(SHorizontalBox)
+		+ SHorizontalBox::Slot()
+		.AutoWidth()
+		.Padding(2.0, 0.0)
+		.VAlign(VAlign_Center)
+		[
+			SNew(SLiveLinkHubTopologyModeSwitcher)
+		];
+
+	GlobalWidgetsSection.AddEntry(
+		FToolMenuEntry::InitWidget(TEXT("TestWidget"), GlobalWidget, FText::GetEmpty(), true, false)
+	);
+	GlobalWidgetsSection.AddSeparator(NAME_None);
+
+	FToolMenuSection& SessionConfigSection = Menu->AddSection(TEXT("ModeWidgets"), FText::GetEmpty(), FToolMenuInsert(TEXT("GlobalWidgets"), EToolMenuInsertType::After));
+
+	TSharedRef<SHorizontalBox> ModeWidgets = SNew(SHorizontalBox);
+
+	TArray<TSharedRef<SWidget>> Widgets = GetStatusBarWidgets_Impl();
 	for (const TSharedRef<SWidget>& Widget : Widgets)
 	{
-		SessionConfigSection.AddEntry(
-			FToolMenuEntry::InitWidget(TEXT("ModeWidget"), Widget, FText::GetEmpty(), true, false)
-		);
+		ModeWidgets->AddSlot()
+			.AutoWidth()
+			.Padding(FMargin(0.0, 0.0, 0.0, 1.0))
+			.VAlign(VAlign_Center)
+			[
+				Widget
+			];
 	}
+
+	SessionConfigSection.AddEntry(
+		FToolMenuEntry::InitWidget(TEXT("ModeWidget"), ModeWidgets, FText::GetEmpty(), true, false)
+	);
 }
 
 void FLiveLinkHubApplicationMode::PreDeactivateMode()
diff --git a/Engine/Plugins/Animation/LiveLinkHub/Source/LiveLinkHub/Private/LiveLinkHubCreatorAppMode.h b/Engine/Plugins/Animation/LiveLinkHub/Source/LiveLinkHub/Private/LiveLinkHubCreatorAppMode.h
--- a/Engine/Plugins/Animation/LiveLinkHub/Source/LiveLinkHub/Private/LiveLinkHubCreatorAppMode.h
+++ b/Engine/Plugins/Animation/LiveLinkHub/Source/LiveLinkHub/Private/LiveLinkHubCreatorAppMode.h
@@ -301,7 +301,7 @@ class FLiveLinkHubCreatorAppMode : public FLiveLinkHubApplicationMode
 {
 public:
 	FLiveLinkHubCreatorAppMode(TSharedPtr<FLiveLinkHubApplicationBase> App)
-		: FLiveLinkHubApplicationMode("CreatorMode", LOCTEXT("CreatorModeLabel", "Creator Mode"), App)
+		: FLiveLinkHubApplicationMode("CreatorMode", LOCTEXT("CreatorModeLabel", "Live Data"), App)
 	{
 
 		TabLayout = FTabManager::NewLayout("LiveLinkHubCreatorMode_v1.1");
@@ -381,8 +381,6 @@ public:
 			)
 		);
 		
-		ToolbarExtender->AddToolBarExtension("Asset", EExtensionHook::After, App->GetToolkitCommands(), FToolBarExtensionDelegate::CreateRaw(this, &FLiveLinkHubCreatorAppMode::AddRecordingInfoToolbarExtension));
-
 		TabFactories.RegisterFactory(MakeShared<FLiveLinkHubSourcesTabSummoner>(App, PanelController));
 		TabFactories.RegisterFactory(MakeShared<FLiveLinkHubSourceDetailsTabSummoner>(App, PanelController));
 		TabFactories.RegisterFactory(MakeShared<FLiveLinkHubSubjectsTabSummoner>(App, PanelController));
@@ -413,6 +411,7 @@ public:
 	{
 		FLiveLinkHubModule& LiveLinkHubModule = FModuleManager::Get().GetModuleChecked<FLiveLinkHubModule>("LiveLinkHub");
 		return {
+			SNew(SBox).Padding(FMargin(0,0,20,0)) [ SNew(SLiveLinkRecordingSessionInfo) ],
 			LiveLinkHubModule.GetRecordingController()->MakeRecordToolbarEntry(),
 			SNew(SLiveLinkTimecode)
 		}; 
@@ -422,6 +421,7 @@ public:
 	{
 		return {
 			SNew(STextBlock)
+			.Margin(FMargin(0.0, 0.0, 4.0, 0.0))
 			.Text_Raw(this, &FLiveLinkHubCreatorAppMode::GetLoadedConfigText)
 		};
 	}
@@ -454,21 +454,6 @@ private:
 		return FText::GetEmpty();
 	}
 
-	/** Adds the take and slate widgets to the toolbar.  */
-	void AddRecordingInfoToolbarExtension(FToolBarBuilder& Builder)
-	{
-		Builder.AddWidget(
-			SNew(SHorizontalBox)
-			+ SHorizontalBox::Slot()
-			.FillWidth(1.0)
-			.VAlign(VAlign_Fill)
-			[
-				SNew(SLiveLinkRecordingSessionInfo)
-			], /*TutorialHighlightName*/NAME_None, /*bSearchable*/ false, EHorizontalAlignment::HAlign_Right
-		);
-	}
-
-
 protected:
 
 	/** Holds the livelink panel controller responsible for creating sources and subjects tabs. */
diff --git a/Engine/Plugins/Animation/LiveLinkHub/Source/LiveLinkHub/Private/UI/Widgets/SLiveLinkHubTopologyModeSwitcher.h b/Engine/Plugins/Animation/LiveLinkHub/Source/LiveLinkHub/Private/UI/Widgets/SLiveLinkHubTopologyModeSwitcher.h
--- a/Engine/Plugins/Animation/LiveLinkHub/Source/LiveLinkHub/Private/UI/Widgets/SLiveLinkHubTopologyModeSwitcher.h
+++ b/Engine/Plugins/Animation/LiveLinkHub/Source/LiveLinkHub/Private/UI/Widgets/SLiveLinkHubTopologyModeSwitcher.h
@@ -6,6 +6,10 @@
 #include "Widgets/SCompoundWidget.h"
 
 #include "LiveLinkHub.h"
+#include "Styling/ISlateStyle.h"
+#include "Styling/SlateStyleRegistry.h"
+#include "Styling/ToolBarStyle.h"
+#include "Widgets/Input/SComboButton.h"
 #include "Widgets/Images/SImage.h"
 #include "Widgets/Input/SButton.h"
 
@@ -33,19 +37,38 @@ public:
 			ModeTooltip = FText::Format(LOCTEXT("ReadOnlyToolTip", "{0}\nNote: Since the mode was set by a command line, it cannot be changed at runtime."), ModeTooltip);
 		}
 
+		HubModeIcon = FSlateStyleRegistry::FindSlateStyle("LiveLinkStyle")->GetBrush("LiveLinkHub.HubMode");
+		SpokeModeIcon = FSlateStyleRegistry::FindSlateStyle("LiveLinkStyle")->GetBrush("LiveLinkHub.SpokeMode");
+
+		const FToolBarStyle& ToolbarStyle = FAppStyle::Get().GetWidgetStyle<FToolBarStyle>("SlimToolbar");
+
 		ChildSlot
 		[
-			SAssignNew(ModeButton, SButton)
-			.ContentPadding(FMargin(0.0f, 2.0f))
-			.HAlign(HAlign_Center)
-			.VAlign(VAlign_Center)
-			.IsEnabled(!bReadOnly)
+			SAssignNew(ModeButton, SComboButton)
+			.ContentPadding(FMargin(0.0))
+			.ButtonStyle(&ToolbarStyle.ButtonStyle)
+			.ComboButtonStyle(&ToolbarStyle.ComboButtonStyle)
+			.ForegroundColor(FSlateColor::UseStyle())
 			.ToolTipText(ModeTooltip)
-			.OnClicked(this, &SLiveLinkHubTopologyModeSwitcher::OnModeClicked)
+			.ButtonContent()
 			[
-				SAssignNew(ButtonContent, STextBlock)
-				.Text(GetModeLabel(Mode))
+				SNew(SHorizontalBox)
+				+ SHorizontalBox::Slot()
+				.AutoWidth()
+				[
+					SNew(SImage)
+					.DesiredSizeOverride(FVector2D{16, 16})
+					.Image(this, &SLiveLinkHubTopologyModeSwitcher::GetModeIcon)
+				]
+				+ SHorizontalBox::Slot()
+				.Padding(FMargin(4.0, 0.0, 0.0, 0.0))
+				.AutoWidth()
+				[
+					SAssignNew(ButtonContent, STextBlock)
+					.Text(UEnum::GetDisplayValueAsText(Mode))
+				]
 			]
+			.OnGetMenuContent(this, &SLiveLinkHubTopologyModeSwitcher::OnGetMenuContent)
 		];
 
 		LiveLinkHub->OnTopologyModeChanged().AddSP(this, &SLiveLinkHubTopologyModeSwitcher::OnModeChanged);
@@ -78,14 +101,63 @@ private:
 			OtherMode = ELiveLinkTopologyMode::Hub;
 		}
 		
-		ButtonContent->SetText(GetModeLabel(Mode));
+		ButtonContent->SetText(UEnum::GetDisplayValueAsText(Mode));
 		ModeButton->SetToolTipText(FText::Format(LOCTEXT("ButtonContentToolTip", "{0}\nClick to change to {1} mode"), GetModeTooltip(Mode), UEnum::GetDisplayValueAsText(OtherMode)));
 	}
 
-	/** Returns the label for the mode button. */
-	static FText GetModeLabel(ELiveLinkTopologyMode InMode)
+	const FSlateBrush* GetModeIcon() const
 	{
-		return FText::Format(LOCTEXT("ModeLabel", "Mode: {0}"), UEnum::GetDisplayValueAsText(InMode));
+		const ELiveLinkTopologyMode Mode =  FLiveLinkHub::Get()->GetTopologyMode();
+		if (Mode == ELiveLinkTopologyMode::Spoke)
+		{
+			return SpokeModeIcon;
+		}
+		else
+		{
+			return HubModeIcon;
+		}
+	}
+
+	TSharedRef<SWidget> OnGetMenuContent() const
+	{
+
+		const bool CloseAfterSelection = true;
+		FMenuBuilder MenuBuilder(CloseAfterSelection, NULL);
+
+		MenuBuilder.BeginSection("OperationMode", LOCTEXT("OperationModeLabel", "Operation Mode"));
+
+		MenuBuilder.AddMenuEntry(
+			LOCTEXT("SpokeModeEntry", "Spoke"),
+			GetModeTooltip(ELiveLinkTopologyMode::Spoke),
+			FSlateIcon("LiveLinkStyle", "LiveLinkHub.SpokeMode"),
+			FExecuteAction::CreateStatic([]()
+			{
+				const ELiveLinkTopologyMode Mode = FLiveLinkHub::Get()->GetTopologyMode();
+				if (Mode != ELiveLinkTopologyMode::Spoke)
+				{
+					FLiveLinkHub::Get()->ToggleTopologyMode();
+				}
+			}),
+			NAME_None,
+			EUserInterfaceActionType::Button);
+
+		MenuBuilder.AddMenuEntry(
+			LOCTEXT("HubModeEntry", "Hub"),
+			GetModeTooltip(ELiveLinkTopologyMode::Hub),
+			FSlateIcon("LiveLinkStyle", "LiveLinkHub.HubMode"),
+			FExecuteAction::CreateStatic([]()
+			{
+				const ELiveLinkTopologyMode Mode = FLiveLinkHub::Get()->GetTopologyMode();
+				if (Mode != ELiveLinkTopologyMode::Hub)
+				{
+					FLiveLinkHub::Get()->ToggleTopologyMode();
+				}
+			}),
+			NAME_None,
+			EUserInterfaceActionType::Button);
+		MenuBuilder.EndSection();
+
+		return MenuBuilder.MakeWidget();
 	}
 
 	/** Returns the tooltip for a mode. */
@@ -98,9 +170,13 @@ private:
 
 private:
 	/** Button that's responsible for switching the topology mode of the app. */
-	TSharedPtr<SButton> ModeButton;
+	TSharedPtr<SComboButton> ModeButton;
 	/** Text indicating the current mode. */
 	TSharedPtr<STextBlock> ButtonContent;
+	/** Icon for Hub mode. */
+	const FSlateBrush* HubModeIcon = nullptr;
+	/** Icon for Spoke mode. */
+	const FSlateBrush* SpokeModeIcon = nullptr;
 };
 
 #undef LOCTEXT_NAMESPACE
\ No newline at end of file
diff --git a/Engine/Plugins/Animation/LiveLinkHub/Source/LiveLinkHub/Private/UI/Widgets/SLiveLinkRecordingSessionInfo.cpp b/Engine/Plugins/Animation/LiveLinkHub/Source/LiveLinkHub/Private/UI/Widgets/SLiveLinkRecordingSessionInfo.cpp
--- a/Engine/Plugins/Animation/LiveLinkHub/Source/LiveLinkHub/Private/UI/Widgets/SLiveLinkRecordingSessionInfo.cpp
+++ b/Engine/Plugins/Animation/LiveLinkHub/Source/LiveLinkHub/Private/UI/Widgets/SLiveLinkRecordingSessionInfo.cpp
@@ -14,7 +14,7 @@
 
 void SLiveLinkRecordingSessionInfo::Construct(const FArguments& InArgs)
 {
-	const double IntraElementPadding = 6.0;
+	const FMargin IntraElementPadding = FMargin(6.0, 4.0);
 
 	ILiveLinkRecordingSessionInfo& SessionInfo = ILiveLinkRecordingSessionInfo::Get();
 
@@ -26,7 +26,7 @@ void SLiveLinkRecordingSessionInfo::Construct(const FArguments& InArgs)
 		.AutoWidth()
 		[
 			SNew(STextBlock)
-			.Margin(FMargin(IntraElementPadding, 0.0))
+			.Margin(IntraElementPadding)
 			.Text(LOCTEXT("SessionEditLabel", "Session"))
 		]
 		+ SHorizontalBox::Slot()
@@ -35,6 +35,7 @@ void SLiveLinkRecordingSessionInfo::Construct(const FArguments& InArgs)
 		[
 			SNew(SEditableTextBox)
 			.MinDesiredWidth(60)
+			.Padding(IntraElementPadding)
 			.Text_Lambda([&SessionInfo] { return FText::FromString(SessionInfo.GetSessionName()); })
 			.OnTextCommitted_Lambda([&SessionInfo]
 				(const FText& InText, ETextCommit::Type)
@@ -47,7 +48,7 @@ void SLiveLinkRecordingSessionInfo::Construct(const FArguments& InArgs)
 		.AutoWidth()
 		[
 			SNew(STextBlock)
-			.Margin(FMargin(IntraElementPadding, 0.0))
+			.Margin(IntraElementPadding)
 			.Text(LOCTEXT("SlateEditLabel", "Slate"))
 		]
 		+ SHorizontalBox::Slot()
@@ -55,6 +56,7 @@ void SLiveLinkRecordingSessionInfo::Construct(const FArguments& InArgs)
 		.AutoWidth()
 		[
 			SNew(SEditableTextBox)
+			.Padding(IntraElementPadding)
 			.MinDesiredWidth(60)
 			.Text_Lambda([&SessionInfo] { return FText::FromString(SessionInfo.GetSlateName()); })
 			.OnTextCommitted_Lambda([&SessionInfo]
@@ -68,7 +70,7 @@ void SLiveLinkRecordingSessionInfo::Construct(const FArguments& InArgs)
 		.AutoWidth()
 		[
 			SNew(STextBlock)
-			.Margin(FMargin(IntraElementPadding, 0.0))
+			.Margin(IntraElementPadding)
 			.Text(LOCTEXT("TakeEditLabel", "Take"))
 		]
 		+ SHorizontalBox::Slot()
@@ -77,6 +79,7 @@ void SLiveLinkRecordingSessionInfo::Construct(const FArguments& InArgs)
 		[
 			SNew(SEditableTextBox)
 			.MinDesiredWidth(20)
+			.Padding(IntraElementPadding)
 			.Text_Lambda([&SessionInfo] { return FText::FromString(LexToString(SessionInfo.GetTakeNumber())); })
 			.OnTextCommitted_Lambda(
 				[&SessionInfo]
@@ -90,7 +93,7 @@ void SLiveLinkRecordingSessionInfo::Construct(const FArguments& InArgs)
 		+ SHorizontalBox::Slot()
 		.VAlign(VAlign_Center)
 		.AutoWidth()
-		.Padding(IntraElementPadding, 0.0)
+		.Padding(IntraElementPadding)
 		[
 			SNew(SPositiveActionButton)
 			.Icon(FAppStyle::Get().GetBrush("Icons.Plus"))
diff --git a/Engine/Plugins/Animation/LiveLinkHub/Source/LiveLinkHub/Private/UI/Widgets/SLiveLinkTimecode.cpp b/Engine/Plugins/Animation/LiveLinkHub/Source/LiveLinkHub/Private/UI/Widgets/SLiveLinkTimecode.cpp
--- a/Engine/Plugins/Animation/LiveLinkHub/Source/LiveLinkHub/Private/UI/Widgets/SLiveLinkTimecode.cpp
+++ b/Engine/Plugins/Animation/LiveLinkHub/Source/LiveLinkHub/Private/UI/Widgets/SLiveLinkTimecode.cpp
@@ -4,10 +4,12 @@
 
 #include "EditorFontGlyphs.h"
 #include "Framework/MultiBox/MultiBoxBuilder.h"
+#include "ILiveLinkModule.h"
 #include "Internationalization/Text.h"
 #include "ISettingsModule.h"
 #include "Misc/App.h"
 #include "Modules/ModuleManager.h"
+#include "Styling/SlateStyle.h"
 #include "Styling/SlateStyleRegistry.h"
 #include "Styling/ToolBarStyle.h"
 #include "Widgets/Images/SImage.h"
@@ -19,29 +21,32 @@
 
 #define LOCTEXT_NAMESPACE "LiveLinkHub"
 
-FSlateColor SLiveLinkTimecode::GetStatusColor() const
+FSlateColor SLiveLinkTimecode::GetIconColor() const
 {
 	const ULiveLinkHubTimeAndSyncSettings* TimeAndSyncSettinsgs = GetDefault<ULiveLinkHubTimeAndSyncSettings>();
-	const bool bTimecodeSourceEnabled = TimeAndSyncSettinsgs->bUseLiveLinkHubAsTimecodeSource;
 	const bool bCustomTimeStepEnabled = TimeAndSyncSettinsgs->bUseLiveLinkHubAsCustomTimeStepSource;
 
+	TSharedPtr<FSlateStyleSet> Style = ILiveLinkModule::Get().GetStyle();
 	if (bCustomTimeStepEnabled)
 	{
-		// Get synchronization state? 
-		if (!TimeAndSyncSettinsgs->IsCustomTimeStepValid())
-		{
-			return FSlateColor(FColor::Red);
-		}
+		return TimeAndSyncSettinsgs->IsCustomTimeStepValid() ? Style->GetSlateColor("LiveLink.Color.Valid") : Style->GetSlateColor("LiveLink.Color.Error");
 	}
+	return FSlateColor(FLinearColor::Gray);
+}
+
+FSlateColor SLiveLinkTimecode::GetStatusColor() const
+{
+	const ULiveLinkHubTimeAndSyncSettings* TimeAndSyncSettinsgs = GetDefault<ULiveLinkHubTimeAndSyncSettings>();
+	const bool bTimecodeSourceEnabled = TimeAndSyncSettinsgs->bUseLiveLinkHubAsTimecodeSource;
+
+	TSharedPtr<FSlateStyleSet> Style = ILiveLinkModule::Get().GetStyle();
 
 	if (bTimecodeSourceEnabled)
 	{
-		return TimeAndSyncSettinsgs->IsTimecodeProviderValid() ? FSlateColor(FColor::Green) : FSlateColor(FColor::Red);
-	}
-	else
-	{
-		return FSlateColor(FLinearColor::Gray);
+		return TimeAndSyncSettinsgs->IsTimecodeProviderValid() ? Style->GetSlateColor("LiveLink.Color.Valid") : Style->GetSlateColor("LiveLink.Color.Error");
 	}
+		
+	return FSlateColor(FLinearColor::Gray);
 }
 
 FText SLiveLinkTimecode::GetTimecodeTooltip() const
@@ -89,10 +94,9 @@ void SLiveLinkTimecode::Construct(const FArguments& InArgs)
 		.VAlign(VAlign_Center)
 		.Padding(5, 0, 3, 0)
 		[
-			SNew(STextBlock)
-				.Font(FAppStyle::Get().GetFontStyle("FontAwesome.8"))
-				.ColorAndOpacity(this, &SLiveLinkTimecode::GetStatusColor)
-				.Text(FEditorFontGlyphs::Circle)
+			SNew(SImage)
+			.Image(FSlateIcon("LiveLinkStyle", "LiveLinkHub.TimecodeGenlock").GetIcon())
+			.ColorAndOpacity(this, &SLiveLinkTimecode::GetIconColor)
 		]
 		+ SHorizontalBox::Slot()
 		.AutoWidth()
diff --git a/Engine/Plugins/Animation/LiveLinkHub/Source/LiveLinkHub/Private/UI/Widgets/SLiveLinkTimecode.h b/Engine/Plugins/Animation/LiveLinkHub/Source/LiveLinkHub/Private/UI/Widgets/SLiveLinkTimecode.h
--- a/Engine/Plugins/Animation/LiveLinkHub/Source/LiveLinkHub/Private/UI/Widgets/SLiveLinkTimecode.h
+++ b/Engine/Plugins/Animation/LiveLinkHub/Source/LiveLinkHub/Private/UI/Widgets/SLiveLinkTimecode.h
@@ -23,6 +23,9 @@ public:
 	void Construct(const FArguments& InArgs);
 
 private:
+	/** Returns the color of the icon depending on whether timecode and genlock is set up correctly. */
+	FSlateColor GetIconColor() const;
+
 	/** Indicate the status of the time code send status */
 	FSlateColor GetStatusColor() const;
 
diff --git a/Engine/Plugins/Animation/LiveLinkHub/Source/LiveLinkHub/Private/UI/Window/LiveLinkHubWindowController.cpp b/Engine/Plugins/Animation/LiveLinkHub/Source/LiveLinkHub/Private/UI/Window/LiveLinkHubWindowController.cpp
--- a/Engine/Plugins/Animation/LiveLinkHub/Source/LiveLinkHub/Private/UI/Window/LiveLinkHubWindowController.cpp
+++ b/Engine/Plugins/Animation/LiveLinkHub/Source/LiveLinkHub/Private/UI/Window/LiveLinkHubWindowController.cpp
@@ -119,12 +119,6 @@ void FLiveLinkHubWindowController::RestoreLayout(TSharedPtr<FAssetEditorToolkit>
 		+ SHorizontalBox::Slot()
 		.AutoWidth()
 		.Padding(4.0f)
-		[
-			SNew(SLiveLinkHubTopologyModeSwitcher)
-		]
-		+ SHorizontalBox::Slot()
-		.AutoWidth()
-		.Padding(4.0f)
 		[
 			SNew(SLiveLinkHubMemoryStats)
 		];
diff --git a/Engine/Plugins/VirtualProduction/CaptureManager/CaptureManagerApp/Source/CaptureManagerEditor/Private/UI/LiveLinkHubCaptureManagerMode.cpp b/Engine/Plugins/VirtualProduction/CaptureManager/CaptureManagerApp/Source/CaptureManagerEditor/Private/UI/LiveLinkHubCaptureManagerMode.cpp
--- a/Engine/Plugins/VirtualProduction/CaptureManager/CaptureManagerApp/Source/CaptureManagerEditor/Private/UI/LiveLinkHubCaptureManagerMode.cpp
+++ b/Engine/Plugins/VirtualProduction/CaptureManager/CaptureManagerApp/Source/CaptureManagerEditor/Private/UI/LiveLinkHubCaptureManagerMode.cpp
@@ -179,12 +179,8 @@ FLiveLinkHubCaptureManagerMode::~FLiveLinkHubCaptureManagerMode() = default;
 TArray<TSharedRef<SWidget>> FLiveLinkHubCaptureManagerMode::GetStatusBarWidgets_Impl()
 {
 	return {
-		SNew(SBox)
-		.Padding(0, 4.0, 0, 0)
-		[
-			SNew(STextBlock)
-			.Text(this, &FLiveLinkHubCaptureManagerMode::GetDiscoveredClientsText)
-		]
+		SNew(STextBlock)
+		.Text(this, &FLiveLinkHubCaptureManagerMode::GetDiscoveredClientsText)
 	};
 }
 
diff --git a/Engine/Source/Runtime/LiveLinkMessageBusFramework/Public/LiveLinkMessages.h b/Engine/Source/Runtime/LiveLinkMessageBusFramework/Public/LiveLinkMessages.h
--- a/Engine/Source/Runtime/LiveLinkMessageBusFramework/Public/LiveLinkMessages.h
+++ b/Engine/Source/Runtime/LiveLinkMessageBusFramework/Public/LiveLinkMessages.h
@@ -20,9 +20,9 @@ struct LIVELINKMESSAGEBUSFRAMEWORK_API FLiveLinkMessageAnnotation
 UENUM()
 enum class ELiveLinkTopologyMode : uint8
 {
-	Hub   			UMETA(ToolTip = "LiveLinkHub instance can receive data from other hub instances in Spoke mode and transmit it to Unreal Editor instances."),
-	Spoke 			UMETA(ToolTip = "LiveLinkHub instance can only transmit data it to a LiveLinkHub instance in Hub mode."),
-	UnrealClient 	UMETA(ToolTip = "Unreal Editor client that can only transmit data it to a LiveLinkHub instance in Hub mode."),
+	Hub   			UMETA(ToolTip = "LiveLinkHub instance that can receive data from other hub instances and transmit it to Unreal Editor instances."),
+	Spoke 			UMETA(ToolTip = "LiveLinkHub instance can only transmit data to other LiveLinkHub instances."),
+	UnrealClient 	UMETA(ToolTip = "Unreal Editor instance that can only receive data from LiveLinkHub"),
 	External 		UMETA(ToolTip = "Client that can transmit to Hubs, Spokes and UnrealClients. (Default mode for message bus providers)")
 };
 
