LiveLinkHub - Default playback mode to looping

#rnx

diff --git a/Engine/Plugins/Animation/LiveLinkHub/Source/LiveLinkHub/Private/Recording/LiveLinkHubPlaybackController.h b/Engine/Plugins/Animation/LiveLinkHub/Source/LiveLinkHub/Private/Recording/LiveLinkHubPlaybackController.h
--- a/Engine/Plugins/Animation/LiveLinkHub/Source/LiveLinkHub/Private/Recording/LiveLinkHubPlaybackController.h
+++ b/Engine/Plugins/Animation/LiveLinkHub/Source/LiveLinkHub/Private/Recording/LiveLinkHubPlaybackController.h
@@ -234,7 +234,7 @@ private:
 	/** Preset used to rollback the hub to its previous state after playing a recording. */
 	TStrongObjectPtr<ULiveLinkPreset> RollbackPreset;
 	/** Atomic bool keeping track of whether we should loop the playback. */
-	std::atomic<bool> bLoopPlayback = false;
+	std::atomic<bool> bLoopPlayback = true;
 	/** Implementation of the playback functionality. */
 	TUniquePtr<ILiveLinkRecordingPlayer> RecordingPlayer;
 
