LiveLinkHub - Fix being able to change topology mode when launched with command line argument

#rnx

diff --git a/Engine/Plugins/Animation/LiveLinkHub/Source/LiveLinkHub/Private/UI/Widgets/SLiveLinkHubTopologyModeSwitcher.h b/Engine/Plugins/Animation/LiveLinkHub/Source/LiveLinkHub/Private/UI/Widgets/SLiveLinkHubTopologyModeSwitcher.h
--- a/Engine/Plugins/Animation/LiveLinkHub/Source/LiveLinkHub/Private/UI/Widgets/SLiveLinkHubTopologyModeSwitcher.h
+++ b/Engine/Plugins/Animation/LiveLinkHub/Source/LiveLinkHub/Private/UI/Widgets/SLiveLinkHubTopologyModeSwitcher.h
@@ -42,6 +42,8 @@ public:
 
 		const FToolBarStyle& ToolbarStyle = FAppStyle::Get().GetWidgetStyle<FToolBarStyle>("SlimToolbar");
 
+		const bool bCanSetTopologyMode = FLiveLinkHub::Get()->CanSetTopologyMode();
+
 		ChildSlot
 		[
 			SAssignNew(ModeButton, SComboButton)
@@ -50,6 +52,7 @@ public:
 			.ComboButtonStyle(&ToolbarStyle.ComboButtonStyle)
 			.ForegroundColor(FSlateColor::UseStyle())
 			.ToolTipText(ModeTooltip)
+			.IsEnabled(bCanSetTopologyMode)
 			.ButtonContent()
 			[
 				SNew(SHorizontalBox)
@@ -126,18 +129,28 @@ private:
 
 		MenuBuilder.BeginSection("OperationMode", LOCTEXT("OperationModeLabel", "Operation Mode"));
 
+		const bool bCanSetSpokeMode = FLiveLinkHub::Get()->GetTopologyMode() != ELiveLinkTopologyMode::Spoke;
+		const bool bCanSetHubMode = FLiveLinkHub::Get()->GetTopologyMode() != ELiveLinkTopologyMode::Hub;
+
+		FExecuteAction ToggleModeAction = FExecuteAction::CreateLambda([]()
+			{
+				FLiveLinkHub::Get()->ToggleTopologyMode();
+			});
+
+		FUIAction SpokeAction;
+		SpokeAction.CanExecuteAction = FCanExecuteAction::CreateLambda([bCanSetSpokeMode]() { return bCanSetSpokeMode; });
+		SpokeAction.ExecuteAction = ToggleModeAction;
+
+		FUIAction HubAction;
+		HubAction.CanExecuteAction = FCanExecuteAction::CreateLambda([bCanSetHubMode]() { return bCanSetHubMode; });
+		HubAction.ExecuteAction = ToggleModeAction;
+
+
 		MenuBuilder.AddMenuEntry(
 			LOCTEXT("SpokeModeEntry", "Spoke"),
 			GetModeTooltip(ELiveLinkTopologyMode::Spoke),
 			FSlateIcon("LiveLinkStyle", "LiveLinkHub.SpokeMode"),
-			FExecuteAction::CreateStatic([]()
-			{
-				const ELiveLinkTopologyMode Mode = FLiveLinkHub::Get()->GetTopologyMode();
-				if (Mode != ELiveLinkTopologyMode::Spoke)
-				{
-					FLiveLinkHub::Get()->ToggleTopologyMode();
-				}
-			}),
+			MoveTemp(SpokeAction),
 			NAME_None,
 			EUserInterfaceActionType::Button);
 
@@ -145,14 +158,7 @@ private:
 			LOCTEXT("HubModeEntry", "Hub"),
 			GetModeTooltip(ELiveLinkTopologyMode::Hub),
 			FSlateIcon("LiveLinkStyle", "LiveLinkHub.HubMode"),
-			FExecuteAction::CreateStatic([]()
-			{
-				const ELiveLinkTopologyMode Mode = FLiveLinkHub::Get()->GetTopologyMode();
-				if (Mode != ELiveLinkTopologyMode::Hub)
-				{
-					FLiveLinkHub::Get()->ToggleTopologyMode();
-				}
-			}),
+			MoveTemp(HubAction),
 			NAME_None,
 			EUserInterfaceActionType::Button);
 		MenuBuilder.EndSection();
