LiveLinkHub - User cannot suppress "Save As" dialog when doing a recording

#rnx

diff --git a/Engine/Plugins/Animation/LiveLinkHub/Source/LiveLinkHub/Private/Recording/Implementations/LiveLinkUAssetRecorder.cpp b/Engine/Plugins/Animation/LiveLinkHub/Source/LiveLinkHub/Private/Recording/Implementations/LiveLinkUAssetRecorder.cpp
--- a/Engine/Plugins/Animation/LiveLinkHub/Source/LiveLinkHub/Private/Recording/Implementations/LiveLinkUAssetRecorder.cpp
+++ b/Engine/Plugins/Animation/LiveLinkHub/Source/LiveLinkHub/Private/Recording/Implementations/LiveLinkUAssetRecorder.cpp
@@ -264,25 +264,34 @@ bool FLiveLinkUAssetRecorder::GetSavePresetPackageName(FString& OutName)
 	FString UserPackageName;
 	FString NewPackageName;
 
-	// Get destination for asset
-	bool bFilenameValid = false;
-	while (!bFilenameValid)
+	bool bPromptSaveAs = GetDefault<ULiveLinkHubSettings>()->bPromptSaveAsOnRecord;
+	// this isnt valid, check if file valid with the asset registry
+	if (bPromptSaveAs || FPaths::FileExists(AbsoluteFolderPath / UniqueAssetName + TEXT(".uasset")))
 	{
-		if (!OpenSaveDialog(DialogStartPath, DialogStartName, UserPackageName))
+		// Get destination for asset
+		bool bFilenameValid = false;
+		while (!bFilenameValid)
 		{
-			return false;
-		}
+			if (!OpenSaveDialog(DialogStartPath, DialogStartName, UserPackageName))
+			{
+				return false;
+			}
 
-		NewPackageName = FString::Format(*UserPackageName, FormatArgs);
+			NewPackageName = FString::Format(*UserPackageName, FormatArgs);
 
-		{
-			// Ensure we aren't saving over an existing recording as this isn't supported yet... the package isn't saved properly.
-			// todo: We should support saving over/deleting existing recordings. This is simply a low risk fix to a playback crash for 5.5.0.
-			AssetToolsModule.Get().CreateUniqueAssetName(NewPackageName, TEXT(""), NewPackageName, UniqueAssetName);
+			{
+				// Ensure we aren't saving over an existing recording as this isn't supported yet... the package isn't saved properly.
+				// todo: We should support saving over/deleting existing recordings. This is simply a low risk fix to a playback crash for 5.5.0.
+				AssetToolsModule.Get().CreateUniqueAssetName(NewPackageName, TEXT(""), NewPackageName, UniqueAssetName);
+			}
+
+			FText OutError;
+			bFilenameValid = FFileHelper::IsFilenameValidForSaving(NewPackageName, OutError);
 		}
-		
-		FText OutError;
-		bFilenameValid = FFileHelper::IsFilenameValidForSaving(NewPackageName, OutError);
+	}
+	else
+	{
+		NewPackageName = TEXT("/Game") / UniqueAssetName;
 	}
 
 	OutName = MoveTemp(NewPackageName);
diff --git a/Engine/Plugins/Animation/LiveLinkHub/Source/LiveLinkHub/Private/Session/LiveLinkHubSessionData.h b/Engine/Plugins/Animation/LiveLinkHub/Source/LiveLinkHub/Private/Session/LiveLinkHubSessionData.h
--- a/Engine/Plugins/Animation/LiveLinkHub/Source/LiveLinkHub/Private/Session/LiveLinkHubSessionData.h
+++ b/Engine/Plugins/Animation/LiveLinkHub/Source/LiveLinkHub/Private/Session/LiveLinkHubSessionData.h
@@ -22,6 +22,7 @@ public:
 	ULiveLinkHubSessionData()
 	{
 		RecordingSessionName = TEXT("DefaultSession");
+		RecordingSlateName = TEXT("DefaultSlate");
 		RecordingTakeNumber = 1;
 	}
 
diff --git a/Engine/Plugins/Animation/LiveLinkHub/Source/LiveLinkHub/Private/Settings/LiveLinkHubSettings.h b/Engine/Plugins/Animation/LiveLinkHub/Source/LiveLinkHub/Private/Settings/LiveLinkHubSettings.h
--- a/Engine/Plugins/Animation/LiveLinkHub/Source/LiveLinkHub/Private/Settings/LiveLinkHubSettings.h
+++ b/Engine/Plugins/Animation/LiveLinkHub/Source/LiveLinkHub/Private/Settings/LiveLinkHubSettings.h
@@ -87,6 +87,10 @@ public:
 	UPROPERTY(config, EditAnywhere, Category="LiveLinkHub", meta = (ConfigRestartRequired = true, ClampMin="15.0"))
 	float TargetFrameRate = 60.0f;
 
+	/** Whether to prompt the user to pick a save directory after doing a recording. */
+	UPROPERTY(config, EditAnywhere, Category = "LiveLinkHub")
+	bool bPromptSaveAsOnRecord = false;
+
 	/** Maximum time in seconds to wait for sources to clean up. Increase this value if you notice that some sources are incorrectly cleaned up when switching a config. */
 	UPROPERTY(config, EditAnywhere, AdvancedDisplay, Category="LiveLinkHub", meta = (ClampMin="0.0"))
 	float SourceMaxCleanupTime = 0.25f;
