LiveLinkHub - Fix not being able to reset a custom time step subject name

diff --git a/Engine/Plugins/Animation/LiveLinkHub/Source/LiveLinkHub/Private/LiveLinkHubModule.cpp b/Engine/Plugins/Animation/LiveLinkHub/Source/LiveLinkHub/Private/LiveLinkHubModule.cpp
--- a/Engine/Plugins/Animation/LiveLinkHub/Source/LiveLinkHub/Private/LiveLinkHubModule.cpp
+++ b/Engine/Plugins/Animation/LiveLinkHub/Source/LiveLinkHub/Private/LiveLinkHubModule.cpp
@@ -6,6 +6,7 @@
 #include "Clients/LiveLinkHubProvider.h"
 #include "LevelEditor.h"
 #include "LiveLinkHubLog.h"
+#include "LiveLinkHubMessages.h"
 #include "LiveLinkHubSubjectSettings.h"
 #include "LiveLinkSettings.h"
 #include "Misc/App.h"
@@ -16,10 +17,12 @@
 #include "Recording/LiveLinkHubRecordingController.h"
 #include "Settings/LiveLinkHubSettings.h"
 #include "Settings/LiveLinkHubSettingsCustomization.h"
+#include "Settings/LiveLinkHubCustomTimeStepSettingsCustomization.h"
 #include "Settings/LiveLinkSettingsCustomization.h"
 #include "Textures/SlateIcon.h"
 #include "Widgets/Text/STextBlock.h"
 
+
 #define LOCTEXT_NAMESPACE "LiveLinkHubModule"
 
 void FLiveLinkHubModule::StartupModule()
@@ -48,6 +51,7 @@ void FLiveLinkHubModule::StartupModule()
 	FPropertyEditorModule& PropertyModule = FModuleManager::LoadModuleChecked<FPropertyEditorModule>("PropertyEditor");
 
 	PropertyModule.RegisterCustomClassLayout(ULiveLinkHubSettings::StaticClass()->GetFName(), FOnGetDetailCustomizationInstance::CreateStatic(&FLiveLinkHubSettingsCustomization::MakeInstance));
+	PropertyModule.RegisterCustomPropertyTypeLayout(FLiveLinkHubCustomTimeStepSettings::StaticStruct()->GetFName(), FOnGetPropertyTypeCustomizationInstance::CreateStatic(&FLiveLinkHubCustomTimeStepSettingsCustomization::MakeInstance));
 
 	// Apply our customization for core live link settings, only if we aren't running in the full editor. We hide properties that aren't
 	// supported in a standalone application context, but are needed if loaded in the editor.
@@ -66,6 +70,7 @@ void FLiveLinkHubModule::ShutdownModule()
 		PropertyEditorModule->UnregisterCustomClassLayout(ULiveLinkHubSettings::StaticClass()->GetFName());
 		if (bUseSettingsDetailCustomization)
 		{
+			PropertyEditorModule->UnregisterCustomClassLayout(FLiveLinkHubCustomTimeStepSettings::StaticStruct()->GetFName());
 			PropertyEditorModule->UnregisterCustomClassLayout(ULiveLinkSettings::StaticClass()->GetFName());
 		}
 	}
diff --git a/Engine/Plugins/Animation/LiveLinkHub/Source/LiveLinkHub/Private/Settings/LiveLinkHubCustomTimeStepSettingsCustomization.h b/Engine/Plugins/Animation/LiveLinkHub/Source/LiveLinkHub/Private/Settings/LiveLinkHubCustomTimeStepSettingsCustomization.h
new file mode 100644
--- /dev/null
+++ b/Engine/Plugins/Animation/LiveLinkHub/Source/LiveLinkHub/Private/Settings/LiveLinkHubCustomTimeStepSettingsCustomization.h
@@ -0,0 +1,66 @@
+// Copyright Epic Games, Inc. All Rights Reserved.
+
+#pragma once
+
+#include "IPropertyTypeCustomization.h"
+
+#include "IDetailChildrenBuilder.h"
+#include "IDetailPropertyRow.h"
+#include "LiveLinkHubMessages.h"
+#include "CoreMinimal.h"
+#include "PropertyHandle.h"
+
+/**
+ * Customization for the CustomTimeStep settings.
+ */
+class FLiveLinkHubCustomTimeStepSettingsCustomization : public IPropertyTypeCustomization
+{
+public:
+	static TSharedRef<IPropertyTypeCustomization> MakeInstance()
+	{
+		return MakeShared<FLiveLinkHubCustomTimeStepSettingsCustomization>();
+	}
+
+	virtual void CustomizeHeader(TSharedRef<IPropertyHandle> PropertyHandle, FDetailWidgetRow& HeaderRow, IPropertyTypeCustomizationUtils& CustomizationUtils)
+	{
+
+	}
+
+	virtual void CustomizeChildren(TSharedRef<IPropertyHandle> PropertyHandle, IDetailChildrenBuilder& ChildBuilder, IPropertyTypeCustomizationUtils& CustomizationUtils) override
+	{
+		uint32 NumberOfChild;
+		if (PropertyHandle->GetNumChildren(NumberOfChild) == FPropertyAccess::Success)
+		{
+			for (uint32 Index = 0; Index < NumberOfChild; ++Index)
+			{
+				TSharedRef<IPropertyHandle> ChildPropertyHandle = PropertyHandle->GetChildHandle(Index).ToSharedRef();
+
+				if (ChildPropertyHandle->GetProperty()->GetFName() == GET_MEMBER_NAME_CHECKED(FLiveLinkHubCustomTimeStepSettings, SubjectName))
+				{
+					FIsResetToDefaultVisible IsResetToDefaultVisible = FIsResetToDefaultVisible::CreateLambda(
+						[](TSharedPtr<IPropertyHandle> PropertyHandle)
+						{ 
+							FName Value;
+							PropertyHandle->GetChildHandle(GET_MEMBER_NAME_CHECKED(FLiveLinkSubjectName, Name))->GetValue(Value);
+							return Value != NAME_None;
+						});
+
+					FResetToDefaultHandler ResetToDefaultHandler = FResetToDefaultHandler::CreateLambda(
+						[](TSharedPtr<IPropertyHandle> PropertyHandle)
+						{
+							PropertyHandle->GetChildHandle(GET_MEMBER_NAME_CHECKED(FLiveLinkSubjectName, Name))->SetValue(NAME_None);
+						});
+
+					const FResetToDefaultOverride Override = FResetToDefaultOverride::Create(IsResetToDefaultVisible, ResetToDefaultHandler);
+
+					ChildBuilder.AddProperty(ChildPropertyHandle)
+						.OverrideResetToDefault(Override);
+				}
+				else
+				{
+					ChildBuilder.AddProperty(ChildPropertyHandle);
+				}
+			}
+		}
+	}
+};
