Live Link Hub - Fix missing copyright menu

diff --git a/Engine/Plugins/Animation/LiveLinkHub/Source/LiveLinkHub/Private/LiveLinkHub.cpp b/Engine/Plugins/Animation/LiveLinkHub/Source/LiveLinkHub/Private/LiveLinkHub.cpp
--- a/Engine/Plugins/Animation/LiveLinkHub/Source/LiveLinkHub/Private/LiveLinkHub.cpp
+++ b/Engine/Plugins/Animation/LiveLinkHub/Source/LiveLinkHub/Private/LiveLinkHub.cpp
@@ -2,7 +2,7 @@
 
 #include "LiveLinkHub.h"
 
-
+#include "AboutScreen.h"
 #include "AssetRegistry/AssetRegistryModule.h"
 #include "AssetRegistry/IAssetRegistry.h"
 #include "Clients/LiveLinkHubClientsController.h"
@@ -372,6 +372,7 @@ void FLiveLinkHub::MapToolkitCommands()
 	CommandList->MapAction(Commands.SaveConfig, FExecuteAction::CreateSP(this, &FLiveLinkHub::SaveConfig),
 		FCanExecuteAction::CreateSP(this, &FLiveLinkHub::CanSaveConfig));
 	CommandList->MapAction(Commands.OpenLogsFolder, FExecuteAction::CreateSP(this, &FLiveLinkHub::OpenLogsFolder));
+	CommandList->MapAction(Commands.OpenAboutMenu, FExecuteAction::CreateSP(this, &FLiveLinkHub::OpenAboutMenu));
 
 	FLiveLinkHubApplication::MapToolkitCommands();
 }
@@ -448,6 +449,24 @@ void FLiveLinkHub::OpenLogsFolder()
 	FPlatformProcess::ExploreFolder(*FileDirectory);
 }
 
+void FLiveLinkHub::OpenAboutMenu()
+{
+	const FText AboutWindowTitle = LOCTEXT("AboutLiveLinkHub", "About Live Link Hub");
+
+	TSharedPtr<SWindow> AboutWindow =
+		SNew(SWindow)
+		.Title(AboutWindowTitle)
+		.ClientSize(FVector2D(720.f, 538.f))
+		.SupportsMaximize(false).SupportsMinimize(false)
+		.SizingRule(ESizingRule::FixedSize)
+		[
+			SNew(SAboutScreen)
+		];
+
+	FSlateApplication::Get().AddWindow(AboutWindow.ToSharedRef());
+}
+
+
 void FLiveLinkHub::OnModularFeatureRegistered(const FName& Type, IModularFeature* ModularFeature)
 {
 	if (Type == ILiveLinkHubApplicationModeFactory::ModularFeatureName)
diff --git a/Engine/Plugins/Animation/LiveLinkHub/Source/LiveLinkHub/Private/LiveLinkHub.h b/Engine/Plugins/Animation/LiveLinkHub/Source/LiveLinkHub/Private/LiveLinkHub.h
--- a/Engine/Plugins/Animation/LiveLinkHub/Source/LiveLinkHub/Private/LiveLinkHub.h
+++ b/Engine/Plugins/Animation/LiveLinkHub/Source/LiveLinkHub/Private/LiveLinkHub.h
@@ -134,6 +134,9 @@ private:
 	/** Open the folder containing logs for the app. */
 	void OpenLogsFolder();
 
+	/** Open the help menu. */
+	void OpenAboutMenu();
+
 	/** Handles adding application modes that are registered after LLH has finished initialization. */
 	void OnModularFeatureRegistered(const FName& Type, class IModularFeature* ModularFeature);
 
diff --git a/Engine/Plugins/Animation/LiveLinkHub/Source/LiveLinkHub/Private/LiveLinkHubApplication.cpp b/Engine/Plugins/Animation/LiveLinkHub/Source/LiveLinkHub/Private/LiveLinkHubApplication.cpp
--- a/Engine/Plugins/Animation/LiveLinkHub/Source/LiveLinkHub/Private/LiveLinkHubApplication.cpp
+++ b/Engine/Plugins/Animation/LiveLinkHub/Source/LiveLinkHub/Private/LiveLinkHubApplication.cpp
@@ -491,6 +491,8 @@ void FLiveLinkHubApplication::RegisterMainMenu()
 
 	FToolMenuEntry OpenLogsFolderEntry = FToolMenuEntry::InitMenuEntry(FLiveLinkHubCommands::Get().OpenLogsFolder, TAttribute<FText>(), TAttribute<FText>(), FSlateIcon(FAppStyle::GetAppStyleSetName(), "Icons.FolderOpen"));
 	HelpEntry->AddMenuEntry(NAME_None, OpenLogsFolderEntry);
+	FToolMenuEntry AboutMenuEntry = FToolMenuEntry::InitMenuEntry(FLiveLinkHubCommands::Get().OpenAboutMenu, TAttribute<FText>(), TAttribute<FText>(), FSlateIcon(FAppStyle::GetAppStyleSetName(), "Icons.Help"));
+	HelpEntry->AddMenuEntry(NAME_None, AboutMenuEntry);
 
 	WindowEntry->AddDynamicSection(NAME_None, FNewSectionConstructChoice{FNewToolMenuDelegateLegacy::CreateRaw(this, &FLiveLinkHubApplication::FillWindowMenu)});
 }
diff --git a/Engine/Plugins/Animation/LiveLinkHub/Source/LiveLinkHub/Private/LiveLinkHubCommands.h b/Engine/Plugins/Animation/LiveLinkHub/Source/LiveLinkHub/Private/LiveLinkHubCommands.h
--- a/Engine/Plugins/Animation/LiveLinkHub/Source/LiveLinkHub/Private/LiveLinkHubCommands.h
+++ b/Engine/Plugins/Animation/LiveLinkHub/Source/LiveLinkHub/Private/LiveLinkHubCommands.h
@@ -39,6 +39,7 @@ public:
 		UI_COMMAND(SaveConfigAs, "Save Config As...", "Save a config to a new file", EUserInterfaceActionType::Button, FInputChord(EModifierKey::Control | EModifierKey::Alt, EKeys::S));
 		UI_COMMAND(SaveConfig, "Save Config", "Save the current config", EUserInterfaceActionType::Button, FInputChord(EKeys::S, EModifierKey::Control));
 		UI_COMMAND(OpenLogsFolder, "Open Logs Folder", "Open the folder containing the logs.", EUserInterfaceActionType::Button, FInputChord());
+		UI_COMMAND(OpenAboutMenu, "About", "Open the About Live Link Hub window.", EUserInterfaceActionType::Button, FInputChord());
 	}
 
 public:
@@ -52,6 +53,8 @@ public:
 	TSharedPtr<FUICommandInfo> SaveConfig;
 	/** Open the folder containing the logs for live link hub. */
 	TSharedPtr<FUICommandInfo> OpenLogsFolder;
+	/** Open the About window. */
+	TSharedPtr<FUICommandInfo> OpenAboutMenu;
 };
 
 #undef LOCTEXT_NAMESPACE
\ No newline at end of file
