LiveLinkHub - Fix thumbnail size for recording list view

diff --git a/Engine/Plugins/Animation/LiveLinkHub/Source/LiveLinkHub/Private/Recording/SLiveLinkHubRecordingListView.h b/Engine/Plugins/Animation/LiveLinkHub/Source/LiveLinkHub/Private/Recording/SLiveLinkHubRecordingListView.h
--- a/Engine/Plugins/Animation/LiveLinkHub/Source/LiveLinkHub/Private/Recording/SLiveLinkHubRecordingListView.h
+++ b/Engine/Plugins/Animation/LiveLinkHub/Source/LiveLinkHub/Private/Recording/SLiveLinkHubRecordingListView.h
@@ -158,6 +158,7 @@ private:
 			AssetPickerConfig.bCanShowClasses = false;
 			AssetPickerConfig.bShowPathInColumnView = true;
 			AssetPickerConfig.bSortByPathInColumnView = false;
+			AssetPickerConfig.InitialThumbnailSize = EThumbnailSize::Tiny;
 			AssetPickerConfig.AssetShowWarningText = GetNoAssetsWarningText();
 
 			AssetPickerConfig.bForceShowEngineContent = true;
