LiveLinkHub - Fix UE only keeping track of a single LLH instance and displaying an invalid LLH status even if there is an active connection

diff --git a/Engine/Plugins/Animation/LiveLinkHub/Source/LiveLinkHubEditor/Private/SLiveLinkHubEditorStatusBar.cpp b/Engine/Plugins/Animation/LiveLinkHub/Source/LiveLinkHubEditor/Private/SLiveLinkHubEditorStatusBar.cpp
--- a/Engine/Plugins/Animation/LiveLinkHub/Source/LiveLinkHubEditor/Private/SLiveLinkHubEditorStatusBar.cpp
+++ b/Engine/Plugins/Animation/LiveLinkHub/Source/LiveLinkHubEditor/Private/SLiveLinkHubEditorStatusBar.cpp
@@ -150,10 +150,7 @@ FSlateColor SLiveLinkHubEditorStatusBar::GetIconColor() const
 
 void SLiveLinkHubEditorStatusBar::OnSourceRemoved(FGuid SourceId)
 {
-	if (HubSourceId == SourceId)
-	{
-		HubSourceId.Invalidate();
-	}
+	SourceIds.Remove(SourceId);
 
 	RefreshSubjects();
 }
@@ -175,22 +172,25 @@ void SLiveLinkHubEditorStatusBar::OnHubConnectionEstablished(FGuid SourceId)
 {
 	bHasConnectedOnce = true;
 
-	HubSourceId = SourceId;
+	SourceIds.Add(SourceId);
 	ConnectionState = EHubConnectionState::Valid;
 }
 
 void SLiveLinkHubEditorStatusBar::CheckHubConnection()
 {
-	if (!HubSourceId.IsValid())
+	if (SourceIds.IsEmpty())
 	{
 		ConnectionState = Invalid;
 		return;
 	}
 
-	if (LiveLinkClient && LiveLinkClient->IsSourceStillValid(HubSourceId))
+	for (FGuid SourceId : SourceIds)
 	{
-		ConnectionState = EHubConnectionState::Valid;
-		return;
+		if (LiveLinkClient && LiveLinkClient->IsSourceStillValid(SourceId))
+		{
+			ConnectionState = EHubConnectionState::Valid;
+			return;
+		}
 	}
 
 	ConnectionState = EHubConnectionState::Timeout;
diff --git a/Engine/Plugins/Animation/LiveLinkHub/Source/LiveLinkHubEditor/Private/SLiveLinkHubEditorStatusBar.h b/Engine/Plugins/Animation/LiveLinkHub/Source/LiveLinkHubEditor/Private/SLiveLinkHubEditorStatusBar.h
--- a/Engine/Plugins/Animation/LiveLinkHub/Source/LiveLinkHubEditor/Private/SLiveLinkHubEditorStatusBar.h
+++ b/Engine/Plugins/Animation/LiveLinkHub/Source/LiveLinkHubEditor/Private/SLiveLinkHubEditorStatusBar.h
@@ -62,8 +62,8 @@ private:
 	/** Whether the editor was able to connect to a hub, will always show the status bar icon after this. */
 	bool bHasConnectedOnce = false;
 
-	/** Id of the livelink source that represents the LLHub. */
-	FGuid HubSourceId;
+	/** Ids of the livelinkhub sources connected to this editor. */
+	TSet<FGuid> SourceIds;
 
 	/** Current connection state to the hub. */
 	EHubConnectionState ConnectionState = Invalid;
