LiveLinkHub - Fix crash when launching game with LiveLinkHub plugin enabled

diff --git a/Engine/Plugins/Animation/LiveLinkHub/Source/LiveLinkHubMessaging/Private/LiveLinkHubMessagingModule.cpp b/Engine/Plugins/Animation/LiveLinkHub/Source/LiveLinkHubMessaging/Private/LiveLinkHubMessagingModule.cpp
--- a/Engine/Plugins/Animation/LiveLinkHub/Source/LiveLinkHubMessaging/Private/LiveLinkHubMessagingModule.cpp
+++ b/Engine/Plugins/Animation/LiveLinkHub/Source/LiveLinkHubMessaging/Private/LiveLinkHubMessagingModule.cpp
@@ -43,9 +43,11 @@ void FLiveLinkHubMessagingModule::StartupModule()
 	GameThreadEndpoint = FMessageEndpoint::Builder(TEXT("LiveLinkHubDiscoveryEndpoint"))
 		.Handling<FLiveLinkHubDiscoveryMessage>(this, &FLiveLinkHubMessagingModule::HandleDiscoveryMessage);
 
-
-	GameThreadEndpoint->Subscribe<FLiveLinkHubDisconnectMessage>();
-	GameThreadEndpoint->Subscribe<FLiveLinkHubDiscoveryMessage>();
+	if (GameThreadEndpoint)
+	{
+		GameThreadEndpoint->Subscribe<FLiveLinkHubDisconnectMessage>();
+		GameThreadEndpoint->Subscribe<FLiveLinkHubDiscoveryMessage>();
+	}
 #endif
 
 	SourceFilterDelegate = ILiveLinkModule::Get().RegisterMessageBusSourceFilter(FOnLiveLinkShouldDisplaySource::CreateRaw(this, &FLiveLinkHubMessagingModule::OnFilterMessageBusSource));
