LiveLinkHub: Fix play and reverse button tooltips not updating correctly while in playback.

diff --git a/Engine/Plugins/Animation/LiveLinkHub/Source/LiveLinkHub/Private/UI/Widgets/SLiveLinkHubPlaybackWidget.cpp b/Engine/Plugins/Animation/LiveLinkHub/Source/LiveLinkHub/Private/UI/Widgets/SLiveLinkHubPlaybackWidget.cpp
--- a/Engine/Plugins/Animation/LiveLinkHub/Source/LiveLinkHub/Private/UI/Widgets/SLiveLinkHubPlaybackWidget.cpp
+++ b/Engine/Plugins/Animation/LiveLinkHub/Source/LiveLinkHub/Private/UI/Widgets/SLiveLinkHubPlaybackWidget.cpp
@@ -140,8 +140,7 @@ void SLiveLinkHubPlaybackWidget::Construct(const FArguments& InArgs)
 				.OnClicked(this, &SLiveLinkHubPlaybackWidget::OnPlayReversePressed)
 				.ToolTipText_Lambda([this]()
 				{
-					const bool bIsPaused = IsPaused();
-					return bIsPaused ? LOCTEXT("ReverseButton", "Reverse") : LOCTEXT("PauseButton", "Pause");
+					return IsPaused() || !IsPlayingInReverse() ? LOCTEXT("ReverseButton", "Reverse") : LOCTEXT("PauseButton", "Pause");
 				})
 				.ButtonStyle( FAppStyle::Get(), "Animation.PlayControlsButton" )
 				.ContentPadding(0.0f)
@@ -159,8 +158,7 @@ void SLiveLinkHubPlaybackWidget::Construct(const FArguments& InArgs)
 				.OnClicked(this, &SLiveLinkHubPlaybackWidget::OnPlayForwardPressed)
 				.ToolTipText_Lambda([this]()
 				{
-					const bool bIsPaused = IsPaused();
-					return bIsPaused ? LOCTEXT("PlayButton", "Play") : LOCTEXT("PauseButton", "Pause");
+					return IsPaused() || IsPlayingInReverse() ? LOCTEXT("PlayButton", "Play") : LOCTEXT("PauseButton", "Pause");
 				})
 				.ButtonStyle( FAppStyle::Get(), "Animation.PlayControlsButton" )
 				.ContentPadding(0.0f)
