LiveLinkHub - Fix timecode provider not using the evaluation type setting

diff --git a/Engine/Plugins/Animation/LiveLinkHub/Source/LiveLinkHubMessaging/Private/LiveLinkHubMessages.cpp b/Engine/Plugins/Animation/LiveLinkHub/Source/LiveLinkHubMessaging/Private/LiveLinkHubMessages.cpp
--- a/Engine/Plugins/Animation/LiveLinkHub/Source/LiveLinkHubMessaging/Private/LiveLinkHubMessages.cpp
+++ b/Engine/Plugins/Animation/LiveLinkHub/Source/LiveLinkHubMessaging/Private/LiveLinkHubMessages.cpp
@@ -119,7 +119,7 @@ void FLiveLinkHubTimecodeSettings::AssignTimecodeSettingsAsProviderToEngine() co
 			LiveLinkProvider->OverrideFrameRate = DesiredFrameRate;
 			LiveLinkProvider->FrameDelay = FrameDelay;
 			LiveLinkProvider->BufferSize = BufferSize;
-
+			LiveLinkProvider->Evaluation = EvaluationType;
 			GEngine->SetTimecodeProvider(LiveLinkProvider);
 			UE_LOG(LogLiveLinkHubMessages, Display, TEXT("Live Link Timecode provider assigned to %s."), *SubjectName.ToString());
 		}
