LiveLinkHub - Fix LiveLinkHub resetting timecode and genlock upon exiting

diff --git a/Engine/Plugins/Animation/LiveLinkHub/Source/LiveLinkHub/Private/Session/LiveLinkHubSession.h b/Engine/Plugins/Animation/LiveLinkHub/Source/LiveLinkHub/Private/Session/LiveLinkHubSession.h
--- a/Engine/Plugins/Animation/LiveLinkHub/Source/LiveLinkHub/Private/Session/LiveLinkHubSession.h
+++ b/Engine/Plugins/Animation/LiveLinkHub/Source/LiveLinkHub/Private/Session/LiveLinkHubSession.h
@@ -140,12 +140,23 @@ public:
 			Clients = CachedSessionClients;
 		}
 
+		const bool bResetTimecode = GetDefault<ULiveLinkHubTimeAndSyncSettings>()->bUseLiveLinkHubAsTimecodeSource;
+		const bool bResetCustomTimeStep = GetDefault<ULiveLinkHubTimeAndSyncSettings>()->bUseLiveLinkHubAsCustomTimeStepSource;
+
 		if (LiveLinkProvider)
 		{
 			for (const FLiveLinkHubClientId& ClientId : Clients)
 			{
-				LiveLinkProvider->ResetTimecodeSettings(ClientId);
-				LiveLinkProvider->ResetCustomTimeStepSettings(ClientId);
+				if (bResetTimecode)
+				{
+					LiveLinkProvider->ResetTimecodeSettings(ClientId);
+				}
+
+				if (bResetCustomTimeStep)
+				{
+					LiveLinkProvider->ResetCustomTimeStepSettings(ClientId);
+				}
+
 				LiveLinkProvider->DisconnectClient(ClientId);
 			}
 		}
