LiveLinkHub: Fix not properly replaying tracks that had inconsistent framerate or data.

- A track's timestamps are now always used when seeking to frames across tracks rather than only relying on qualified frame time conversions. For efficiency we still convert the desired frame time to a track's local frame time first which often puts as at or near the desired frame.

- Timestamps are now saved and loaded in bulk, rather than with each frame. For old recordings the file is iterated once to load the timestamps, which can be slow.

- Add Upgrade feature which will force load the entire recording into memory and then save it to a new file, allowing users to convert old animations to the new format, and decrease future load times. The recording list will now specify if a version is outdated or not. If a recording is outdated a context menu option "Upgrade Recording" will be available.

- Fix multiple event fires when committing a scrubber change via spin boxes.
- Fix crash when exiting the app while streaming frames.

diff --git a/Engine/Plugins/Animation/LiveLinkHub/Source/LiveLinkHub/Private/LiveLinkHub.cpp b/Engine/Plugins/Animation/LiveLinkHub/Source/LiveLinkHub/Private/LiveLinkHub.cpp
--- a/Engine/Plugins/Animation/LiveLinkHub/Source/LiveLinkHub/Private/LiveLinkHub.cpp
+++ b/Engine/Plugins/Animation/LiveLinkHub/Source/LiveLinkHub/Private/LiveLinkHub.cpp
@@ -11,15 +11,15 @@
 #include "Framework/Application/SlateApplication.h"
 #include "GenericPlatform/GenericPlatformOutputDevices.h"
 #include "HAL/PlatformProcess.h"
-#include "ISettingsModule.h"
 #include "ILiveLinkHubMessagingModule.h"
 #include "ILiveLinkModule.h"
+#include "ISettingsModule.h"
 #include "LiveLinkEditorSettings.h"
 #include "LiveLinkHubClient.h"
 #include "LiveLinkHubCommands.h"
 #include "LiveLinkHubCreatorAppMode.h"
-#include "LiveLinkHubModule.h"
 #include "LiveLinkHubMessageBusSourceSettings.h"
+#include "LiveLinkHubModule.h"
 #include "LiveLinkHubSubjectSettings.h"
 #include "LiveLinkHubTicker.h"
 #include "LiveLinkProviderImpl.h"
@@ -150,6 +150,7 @@ void FLiveLinkHub::Initialize(FLiveLinkHubTicker& Ticker)
 	MessagingModule->SetHostTopologyMode(GetTopologyMode());
 
 	FCoreDelegates::OnPostEngineInit.AddRaw(this, &FLiveLinkHub::InitializePostEngineInit);
+	FCoreDelegates::OnPreExit.AddRaw(this, &FLiveLinkHub::OnEnginePreExit);
 }
 
 void FLiveLinkHub::InitializePostEngineInit()
@@ -219,6 +220,17 @@ void FLiveLinkHub::InitializePostEngineInit()
 	}
 }
 
+void FLiveLinkHub::OnEnginePreExit()
+{
+	if (PlaybackController)
+	{
+		if (ULiveLinkUAssetRecording* Recording = Cast<ULiveLinkUAssetRecording>(PlaybackController->GetRecording().Get()))
+		{
+			Recording->UnloadRecordingData();
+		}
+	}
+}
+
 FLiveLinkHub::~FLiveLinkHub()
 {
 	IModularFeatures::Get().UnregisterModularFeature(IConsoleCommandExecutor::ModularFeatureName(), CommandExecutor.Get());
@@ -234,6 +246,8 @@ FLiveLinkHub::~FLiveLinkHub()
 	LiveLinkHubClient->OnStaticDataReceived_AnyThread().RemoveAll(this);
 
 	IModularFeatures::Get().UnregisterModularFeature(ILiveLinkClient::ModularFeatureName, LiveLinkHubClient.Get());
+
+	FCoreDelegates::OnPreExit.RemoveAll(this);
 }
 
 bool FLiveLinkHub::IsInPlayback() const
@@ -361,6 +375,12 @@ bool FLiveLinkHub::CanSetTopologyMode() const
 	return !OverrideTopologyMode.IsSet();
 }
 
+void FLiveLinkHub::UpgradeAndSaveRecording(ULiveLinkUAssetRecording* InRecording)
+{
+	GetPlaybackController()->OnPlaybackReady().AddSP(this, &FLiveLinkHub::OnPlaybackReadyForUpgrade);
+	GetPlaybackController()->PreparePlayback(InRecording);
+}
+
 void FLiveLinkHub::MapToolkitCommands()
 {
 	const FLiveLinkHubCommands& Commands = FLiveLinkHubCommands::Get();
@@ -466,7 +486,6 @@ void FLiveLinkHub::OpenAboutMenu()
 	FSlateApplication::Get().AddWindow(AboutWindow.ToSharedRef());
 }
 
-
 void FLiveLinkHub::OnModularFeatureRegistered(const FName& Type, IModularFeature* ModularFeature)
 {
 	if (Type == ILiveLinkHubApplicationModeFactory::ModularFeatureName)
@@ -502,6 +521,37 @@ void FLiveLinkHub::OnSubjectOutboundNameModified(const FLiveLinkSubjectKey& Subj
 	LiveLinkProvider->RemoveSubject(*PreviousOutboundName);
 }
 
+void FLiveLinkHub::OnPlaybackReadyForUpgrade()
+{
+	if (GetPlaybackController() && GetPlaybackController()->GetRecording().IsValid())
+	{
+		GetPlaybackController()->OnPlaybackReady().RemoveAll(this);
+
+		// Get the save path before starting the long task.
+		const FString DefaultName = GetPlaybackController()->GetRecording()->GetName() + TEXT("_copy");
+		FString RecordingPath;
+		if (!GetRecordingController()->OpenSaveAsPrompt(RecordingPath, DefaultName))
+		{
+			GetPlaybackController()->EjectAndUnload();
+			return;
+		}
+		
+		ULiveLinkUAssetRecording* CurrentRecording = CastChecked<ULiveLinkUAssetRecording>(GetPlaybackController()->GetRecording().Get());
+		
+		TWeakObjectPtr<ULiveLinkUAssetRecording> CurrentRecordingWeakPtr(CurrentRecording);
+		CurrentRecording->GetOnRecordingFullyLoaded().BindLambda([this, CurrentRecordingWeakPtr, RecordingPath]()
+		{
+			const TStrongObjectPtr<ULiveLinkUAssetRecording> CurrentRecordingPin = CurrentRecordingWeakPtr.Pin();
+			if (CurrentRecordingPin.IsValid() && GetRecordingController())
+			{
+				GetRecordingController()->RecordFromExistingAndSave(CurrentRecordingPin.Get(), RecordingPath);
+			}
+		});
+
+		CurrentRecording->LoadEntireRecordingForUpgrade();
+	}
+}
+
 void FLiveLinkHub::RegisterLiveLinkHubSettings()
 {
 	if (ISettingsModule* SettingsModule = FModuleManager::GetModulePtr<ISettingsModule>("Settings"))
diff --git a/Engine/Plugins/Animation/LiveLinkHub/Source/LiveLinkHub/Private/LiveLinkHub.h b/Engine/Plugins/Animation/LiveLinkHub/Source/LiveLinkHub/Private/LiveLinkHub.h
--- a/Engine/Plugins/Animation/LiveLinkHub/Source/LiveLinkHub/Private/LiveLinkHub.h
+++ b/Engine/Plugins/Animation/LiveLinkHub/Source/LiveLinkHub/Private/LiveLinkHub.h
@@ -19,11 +19,12 @@ class FLiveLinkHubRecordingController;
 class FLiveLinkHubRecordingListController;
 class FLiveLinkHubSubjectController;
 class FLiveLinkHubWindowController;
-struct FLiveLinkSubjectKey;
 class FLiveLinkHubProvider;
 class ILiveLinkHubSessionManager;
 class SWindow;
 class ULiveLinkRole;
+class ULiveLinkUAssetRecording;
+struct FLiveLinkSubjectKey;
 
 DECLARE_TS_MULTICAST_DELEGATE_OneParam(FOnTopologyModeChanged, ELiveLinkTopologyMode /*NewMode*/);
 
@@ -97,6 +98,9 @@ public:
 	/** Delegate called when the topology mode changed. */
 	FOnTopologyModeChanged& OnTopologyModeChanged() { return TopologyModeChangedDelegate; }
 
+	/** Save an existing recording as a new file. */
+	void UpgradeAndSaveRecording(ULiveLinkUAssetRecording* InRecording);
+	
 	//~ FAssetEditorToolkit interface
 	virtual void MapToolkitCommands() override;
 	virtual void OnClose() override;
@@ -105,6 +109,9 @@ private:
 	/** Complete the initialization after core engine systems have been initialized. */
 	void InitializePostEngineInit();
 
+	/** When the engine is about to exit. */
+	void OnEnginePreExit();
+
 	//~ LiveLink Client delegates
 	void OnStaticDataReceived_AnyThread(const FLiveLinkSubjectKey& InSubjectKey, TSubclassOf<ULiveLinkRole> InRole, const FLiveLinkStaticDataStruct& InStaticDataStruct) const;
 	void OnFrameDataReceived_AnyThread(const FLiveLinkSubjectKey& InSubjectKey, const FLiveLinkFrameDataStruct& InFrameDataStruct) const;
@@ -127,7 +134,7 @@ private:
 	
 	/** Save the config to the current file. */
 	void SaveConfig();
-
+	
 	/** Open an existing config. */
 	void OpenConfig();
 
@@ -142,6 +149,9 @@ private:
 
 	/** Delegate called when a LiveLinkHubSubjectSettings modifies an outbound name. We need to notify connected clients of this change. */
 	void OnSubjectOutboundNameModified(const FLiveLinkSubjectKey& SubjectKey, const FString& PreviousOutboundName, const FString& NewOutboundName) const;
+
+	/** Callback when playback is ready because we are performing an upgrade. */
+	void OnPlaybackReadyForUpgrade();
 	
 private:
 	/** Implements the logic to manage the clients tabs. */
diff --git a/Engine/Plugins/Animation/LiveLinkHub/Source/LiveLinkHub/Private/Recording/Implementations/LiveLinkFrameData.cpp b/Engine/Plugins/Animation/LiveLinkHub/Source/LiveLinkHub/Private/Recording/Implementations/LiveLinkFrameData.cpp
--- a/Engine/Plugins/Animation/LiveLinkHub/Source/LiveLinkHub/Private/Recording/Implementations/LiveLinkFrameData.cpp
+++ b/Engine/Plugins/Animation/LiveLinkHub/Source/LiveLinkHub/Private/Recording/Implementations/LiveLinkFrameData.cpp
@@ -47,12 +47,12 @@ namespace UE::LiveLinkHub::FrameData::Private
 		TrimCache();
 	}
 
-	RangeHelpers::Private::TRangeArray<int32> FFrameBufferCache::GetCacheBufferRanges() const
+	RangeHelpers::Private::TRangeArray<int32> FFrameBufferCache::GetCacheBufferRanges(bool bIncludeOffset) const
 	{
 		RangeHelpers::Private::TRangeArray<int32> Results;
 		for (const FLiveLinkRecordingBaseDataContainer& Container : FrameData)
 		{
-			Results.Add(Container.GetBufferedFrames());
+			Results.Add(Container.GetBufferedFrames(bIncludeOffset));
 		}
 
 		return Results;
diff --git a/Engine/Plugins/Animation/LiveLinkHub/Source/LiveLinkHub/Private/Recording/Implementations/LiveLinkFrameData.h b/Engine/Plugins/Animation/LiveLinkHub/Source/LiveLinkHub/Private/Recording/Implementations/LiveLinkFrameData.h
--- a/Engine/Plugins/Animation/LiveLinkHub/Source/LiveLinkHub/Private/Recording/Implementations/LiveLinkFrameData.h
+++ b/Engine/Plugins/Animation/LiveLinkHub/Source/LiveLinkHub/Private/Recording/Implementations/LiveLinkFrameData.h
@@ -20,8 +20,11 @@ namespace UE::LiveLinkHub::FrameData::Private
 		/** Remove data from the cache that is no longer needed. */
 		void CleanCache(const TRange<int32>& InActiveRange);
 
-		/** Return all ranges contained within the cache. */
-		RangeHelpers::Private::TRangeArray<int32> GetCacheBufferRanges() const;
+		/**
+		 * Return all ranges contained within the cache.
+		 * @param bIncludeOffset If the local frame time-to-frame index offset used should be added to the ranges.
+		 */
+		RangeHelpers::Private::TRangeArray<int32> GetCacheBufferRanges(bool bIncludeOffset = false) const;
 
 		/** Look through the cache for a loaded frame. */
 		TSharedPtr<FInstancedStruct> TryGetCachedFrame(const int32 InFrame, double& OutTimestamp);
@@ -103,6 +106,8 @@ namespace UE::LiveLinkHub::FrameData::Private
 		int32 MaxFrames = 0;
 		/** Frame offsets and sizes. [FrameOffsetBytes, FrameSizeBytes] */
 		TArray<TTuple<int64, int32>> FrameDiskSizes;
+		/** All timestamps for the track, whether the frame is loaded or not. */
+		TArray<double> Timestamps;
 		/** Whether the frame size is consistent throughout this animation. */
 		bool bHasConsistentFrameSize = false;
 		/** The last timestamp for this frame data. */
@@ -114,6 +119,81 @@ namespace UE::LiveLinkHub::FrameData::Private
 		/** Current iteration data while buffering. */
 		FFrameBufferIterationData BufferIterationData;
 
+		/**
+		 * Find the closest frame matching the given timestamp.
+		 * @param InTimestamp The timestamp of the frame.
+		 * @param InIndexHint A hint to where the index might be, allowing quicker search times.
+		 */
+		int32 FindFrameFromTimestamp(const double InTimestamp, const int32 InIndexHint) const
+		{
+			const int32 Num = Timestamps.Num();
+			if (Num == 0)
+			{
+				ensure(false);
+				return INDEX_NONE;
+			}
+
+			// Check ranges.
+			if (InTimestamp <= Timestamps[0])
+			{
+				return 0;
+			}
+			if (InTimestamp >= Timestamps[Num - 1])
+			{
+				return Num - 1;
+			}
+
+			const int32 Hint = FMath::Clamp(InIndexHint, 0, Num - 1);
+			
+			// Are we already hitting the target?
+			const double HintTime = Timestamps[Hint];
+			const bool bIsLast = (Hint == Num - 1);
+			if (HintTime <= InTimestamp &&
+				(bIsLast || Timestamps[Hint + 1] > InTimestamp))
+			{
+				return Hint;
+			}
+
+			int32 Low, High;
+			int32 Best = 0;
+			if (HintTime < InTimestamp)
+			{
+				// Search to the right.
+				Low  = Hint + 1;
+				High = Num - 1;
+				Best = Hint;
+			}
+			else
+			{
+				// Search to the left.
+				Low  = 0;
+				High = Hint - 1;
+				Best = 0;
+			}
+
+			while (Low <= High)
+			{
+				const int32 Mid = (Low + High) / 2;
+				const double MidTime = Timestamps[Mid];
+
+				if (MidTime == InTimestamp)
+				{
+					return Mid;
+				}
+				if (MidTime < InTimestamp)
+				{
+					Best = Mid;
+					Low  = Mid + 1; // Go right.
+				}
+				else
+				{
+					High = Mid - 1; // Go left.
+				}
+			}
+
+			return Best;
+		}
+		
 		/** The size in bytes of an animation frame. */
 		int32 GetFrameDiskSize(const int32 InFrameIdx) const
 		{
diff --git a/Engine/Plugins/Animation/LiveLinkHub/Source/LiveLinkHub/Private/Recording/Implementations/LiveLinkRecordingDataContainer.cpp b/Engine/Plugins/Animation/LiveLinkHub/Source/LiveLinkHub/Private/Recording/Implementations/LiveLinkRecordingDataContainer.cpp
--- a/Engine/Plugins/Animation/LiveLinkHub/Source/LiveLinkHub/Private/Recording/Implementations/LiveLinkRecordingDataContainer.cpp
+++ b/Engine/Plugins/Animation/LiveLinkHub/Source/LiveLinkHub/Private/Recording/Implementations/LiveLinkRecordingDataContainer.cpp
@@ -17,14 +17,20 @@ void FLiveLinkRecordingBaseDataContainer::ClearData()
 	RecordedDataStartFrame = 0;
 }
 
-TRange<int32> FLiveLinkRecordingBaseDataContainer::GetBufferedFrames() const
+TRange<int32> FLiveLinkRecordingBaseDataContainer::GetBufferedFrames(bool bIncludeOffset) const
 {
 	if (Timestamps.Num() == 0)
 	{
 		return TRange<int32>::Empty();
 	}
-	const int32 EndFrame = FMath::Max(0, RecordedDataStartFrame + Timestamps.Num() - 1);
-	return UE::LiveLinkHub::RangeHelpers::Private::MakeInclusiveRange(RecordedDataStartFrame, EndFrame);
+	int32 StartFrame = RecordedDataStartFrame;
+	int32 EndFrame = FMath::Max(0, RecordedDataStartFrame + Timestamps.Num() - 1);
+	if (bIncludeOffset)
+	{
+		StartFrame += LocalFrameOffset;
+		EndFrame += LocalFrameOffset;
+	}
+	return UE::LiveLinkHub::RangeHelpers::Private::MakeInclusiveRange(StartFrame, EndFrame);
 }
 
 TSharedPtr<FInstancedStruct> FLiveLinkRecordingBaseDataContainer::TryGetFrame(const int32 InFrame) const
diff --git a/Engine/Plugins/Animation/LiveLinkHub/Source/LiveLinkHub/Private/Recording/Implementations/LiveLinkRecordingDataContainer.h b/Engine/Plugins/Animation/LiveLinkHub/Source/LiveLinkHub/Private/Recording/Implementations/LiveLinkRecordingDataContainer.h
--- a/Engine/Plugins/Animation/LiveLinkHub/Source/LiveLinkHub/Private/Recording/Implementations/LiveLinkRecordingDataContainer.h
+++ b/Engine/Plugins/Animation/LiveLinkHub/Source/LiveLinkHub/Private/Recording/Implementations/LiveLinkRecordingDataContainer.h
@@ -27,14 +27,20 @@ struct FLiveLinkRecordingBaseDataContainer
 	/** The current start frame for RecordedData. */
 	int32 RecordedDataStartFrame = 0;
 
+	/** Offset which was applied locally during frame rate conversions to find the correct timestamp. */
+	int32 LocalFrameOffset = 0;
+
 	/** If this container contains no data. */
 	bool IsEmpty() const;
 
 	/** Empty all data. */
 	void ClearData();
 
-	/** Retrieve the total buffered frames. */
-	TRange<int32> GetBufferedFrames() const;
+	/**
+	 * Retrieve the total buffered frames.
+	 * @param bIncludeOffset If the local frame time-to-frame index offset used should be added to the ranges.
+	 */
+	TRange<int32> GetBufferedFrames(bool bIncludeOffset = false) const;
 
 	/**
 	 * Retrieve a loaded frame.
diff --git a/Engine/Plugins/Animation/LiveLinkHub/Source/LiveLinkHub/Private/Recording/Implementations/LiveLinkUAssetRecorder.cpp b/Engine/Plugins/Animation/LiveLinkHub/Source/LiveLinkHub/Private/Recording/Implementations/LiveLinkUAssetRecorder.cpp
--- a/Engine/Plugins/Animation/LiveLinkHub/Source/LiveLinkHub/Private/Recording/Implementations/LiveLinkUAssetRecorder.cpp
+++ b/Engine/Plugins/Animation/LiveLinkHub/Source/LiveLinkHub/Private/Recording/Implementations/LiveLinkUAssetRecorder.cpp
@@ -85,7 +85,8 @@
 
 		TimeRecordingEnded = FPlatformTime::Seconds();
 
-		SaveRecording();
+		const double RecordingLength = TimeRecordingEnded - TimeRecordingStarted;
+		SaveRecording(MoveTemp(*CurrentRecording), RecordingLength, false);
 		CurrentRecording.Reset();
 	}
 }
@@ -142,65 +143,15 @@
 	return true;
 }
 
-void FLiveLinkUAssetRecorder::RecordBaseData(FLiveLinkRecordingBaseDataContainer& StaticDataContainer, TSharedPtr<FInstancedStruct>&& DataToRecord)
+void FLiveLinkUAssetRecorder::SaveFromExistingRecording(const ULiveLinkRecording* InRecording, const FString& InFilePath)
 {
-	const double TimeNowInSeconds = FPlatformTime::Seconds();
-	StaticDataContainer.RecordedData.Add(MoveTemp(DataToRecord));
-	StaticDataContainer.Timestamps.Add(TimeNowInSeconds - TimeRecordingStarted);
+	FLiveLinkUAssetRecordingData InRecordingData = CastChecked<ULiveLinkUAssetRecording>(InRecording)->RecordingData;
+	SaveRecording(MoveTemp(InRecordingData), InRecording->LengthInSeconds, /*bIncrementTake*/ false,
+		InRecording->RecordingPreset, &InRecording->FrameRate, InFilePath);
+	CurrentRecording.Reset();
 }
 
-void FLiveLinkUAssetRecorder::RecordStaticData(const FLiveLinkSubjectKey& SubjectKey, TSubclassOf<ULiveLinkRole> Role, const FLiveLinkStaticDataStruct& StaticData)
+bool FLiveLinkUAssetRecorder::GetSavePresetPackageName(FString& OutName, const FString& InDefaultName, bool bAlwaysOpenSaveAsDialog) const
 {
-	if (bIsRecording && CurrentRecording)
-	{
-		FLiveLinkRecordingStaticDataContainer& StaticDataContainer = CurrentRecording->StaticData.FindOrAdd(SubjectKey);
-		TSharedPtr<FInstancedStruct> NewData = MakeShared<FInstancedStruct>();
-		NewData->InitializeAs(StaticData.GetStruct(), (uint8*)StaticData.GetBaseData());
-		StaticDataContainer.Role = Role;
-
-		RecordBaseData(StaticDataContainer, MoveTemp(NewData));
-	}
-}
-
-void FLiveLinkUAssetRecorder::RecordFrameData(const FLiveLinkSubjectKey& SubjectKey, const FLiveLinkFrameDataStruct& FrameData)
-{
-	if (bIsRecording && CurrentRecording)
-	{
-		FLiveLinkRecordingBaseDataContainer& FrameDataContainer = CurrentRecording->FrameData.FindOrAdd(SubjectKey);
-		TSharedPtr<FInstancedStruct> NewData = MakeShared<FInstancedStruct>();
-		NewData->InitializeAs(FrameData.GetStruct(), (uint8*)FrameData.GetBaseData());
-
-		RecordBaseData(FrameDataContainer, MoveTemp(NewData));
-	}
-}
-
-bool FLiveLinkUAssetRecorder::OpenSaveDialog(const FString& InDefaultPath, const FString& InNewNameSuggestion, FString& OutPackageName)
-{
-	TSharedRef<SWindow> RootWindow = FModuleManager::Get().GetModuleChecked<FLiveLinkHubModule>("LiveLinkHub").GetLiveLinkHub()->GetRootWindow();
-
-	FSaveAssetDialogConfig SaveAssetDialogConfig;
-	{
-		SaveAssetDialogConfig.DefaultPath = InDefaultPath;
-		SaveAssetDialogConfig.DefaultAssetName = InNewNameSuggestion;
-		SaveAssetDialogConfig.AssetClassNames.Add(ULiveLinkRecording::StaticClass()->GetClassPathName());
-		SaveAssetDialogConfig.ExistingAssetPolicy = ESaveAssetDialogExistingAssetPolicy::AllowButWarn;
-		SaveAssetDialogConfig.DialogTitleOverride = LOCTEXT("SaveLiveLinkRecordingDialogTitle", "Save Live Link Recording");
-		SaveAssetDialogConfig.WindowOverride = RootWindow;
-	}
-
-	FContentBrowserModule& ContentBrowserModule = FModuleManager::LoadModuleChecked<FContentBrowserModule>("ContentBrowser");
-	FString SaveObjectPath = ContentBrowserModule.Get().CreateModalSaveAssetDialog(SaveAssetDialogConfig);
-
-	if (!SaveObjectPath.IsEmpty())
-	{
-		OutPackageName = FPackageName::ObjectPathToPackageName(SaveObjectPath);
-		return true;
-	}
-
-	return false;
-}
-
-bool FLiveLinkUAssetRecorder::GetSavePresetPackageName(FString& OutName)
-{
 	const FDateTime Today = FDateTime::Now();
 
@@ -214,9 +165,10 @@
 	FNamingTokenResultData TemplateData = GEngine->GetEngineSubsystem<UNamingTokensEngineSubsystem>()->EvaluateTokenString(GetDefault<ULiveLinkHubSettings>()->FilenameTemplate,
 		Filter);
 
-	FString DefaultName;
+	FString DefaultName = InDefaultName;
 	FString DefaultFolder;
 	// Split folder path and file name.
+	if (DefaultName.IsEmpty())
 	{
 		FString FormattedString = TemplateData.EvaluatedText.ToString();
 		int32 LastSlashIndex;
@@ -264,7 +216,7 @@
 	FString UserPackageName;
 	FString NewPackageName;
 
-	bool bPromptSaveAs = GetDefault<ULiveLinkHubSettings>()->bPromptSaveAsOnRecord;
+	bool bPromptSaveAs = bAlwaysOpenSaveAsDialog || GetDefault<ULiveLinkHubSettings>()->bPromptSaveAsOnRecord;
 	// this isnt valid, check if file valid with the asset registry
 	if (bPromptSaveAs || FPaths::FileExists(AbsoluteFolderPath / UniqueAssetName + TEXT(".uasset")))
 	{
@@ -298,10 +250,70 @@
 	return true;
 }
 
-void FLiveLinkUAssetRecorder::SaveRecording()
+void FLiveLinkUAssetRecorder::RecordBaseData(FLiveLinkRecordingBaseDataContainer& StaticDataContainer, TSharedPtr<FInstancedStruct>&& DataToRecord)
 {
-	FString PackageName;
-	if (!GetSavePresetPackageName(PackageName))
+	const double TimeNowInSeconds = FPlatformTime::Seconds();
+	StaticDataContainer.RecordedData.Add(MoveTemp(DataToRecord));
+	StaticDataContainer.Timestamps.Add(TimeNowInSeconds - TimeRecordingStarted);
+}
+
+void FLiveLinkUAssetRecorder::RecordStaticData(const FLiveLinkSubjectKey& SubjectKey, TSubclassOf<ULiveLinkRole> Role, const FLiveLinkStaticDataStruct& StaticData)
+{
+	if (bIsRecording && CurrentRecording)
 	{
+		FLiveLinkRecordingStaticDataContainer& StaticDataContainer = CurrentRecording->StaticData.FindOrAdd(SubjectKey);
+		TSharedPtr<FInstancedStruct> NewData = MakeShared<FInstancedStruct>();
+		NewData->InitializeAs(StaticData.GetStruct(), (uint8*)StaticData.GetBaseData());
+		StaticDataContainer.Role = Role;
+
+		RecordBaseData(StaticDataContainer, MoveTemp(NewData));
+	}
+}
+
+void FLiveLinkUAssetRecorder::RecordFrameData(const FLiveLinkSubjectKey& SubjectKey, const FLiveLinkFrameDataStruct& FrameData)
+{
+	if (bIsRecording && CurrentRecording)
+	{
+		FLiveLinkRecordingBaseDataContainer& FrameDataContainer = CurrentRecording->FrameData.FindOrAdd(SubjectKey);
+		TSharedPtr<FInstancedStruct> NewData = MakeShared<FInstancedStruct>();
+		NewData->InitializeAs(FrameData.GetStruct(), (uint8*)FrameData.GetBaseData());
+
+		RecordBaseData(FrameDataContainer, MoveTemp(NewData));
+	}
+}
+
+bool FLiveLinkUAssetRecorder::OpenSaveDialog(const FString& InDefaultPath, const FString& InNewNameSuggestion, FString& OutPackageName) const
+{
+	TSharedRef<SWindow> RootWindow = FModuleManager::Get().GetModuleChecked<FLiveLinkHubModule>("LiveLinkHub").GetLiveLinkHub()->GetRootWindow();
+
+	FSaveAssetDialogConfig SaveAssetDialogConfig;
+	{
+		SaveAssetDialogConfig.DefaultPath = InDefaultPath;
+		SaveAssetDialogConfig.DefaultAssetName = InNewNameSuggestion;
+		SaveAssetDialogConfig.AssetClassNames.Add(ULiveLinkRecording::StaticClass()->GetClassPathName());
+		SaveAssetDialogConfig.ExistingAssetPolicy = ESaveAssetDialogExistingAssetPolicy::AllowButWarn;
+		SaveAssetDialogConfig.DialogTitleOverride = LOCTEXT("SaveLiveLinkRecordingDialogTitle", "Save Live Link Recording");
+		SaveAssetDialogConfig.WindowOverride = RootWindow;
+	}
+
+	FContentBrowserModule& ContentBrowserModule = FModuleManager::LoadModuleChecked<FContentBrowserModule>("ContentBrowser");
+	FString SaveObjectPath = ContentBrowserModule.Get().CreateModalSaveAssetDialog(SaveAssetDialogConfig);
+
+	if (!SaveObjectPath.IsEmpty())
+	{
+		OutPackageName = FPackageName::ObjectPathToPackageName(SaveObjectPath);
+		return true;
+	}
+
+	return false;
+}
+
+void FLiveLinkUAssetRecorder::SaveRecording(FLiveLinkUAssetRecordingData&& InRecordingData,
+	double InRecordingLength, bool bIncrementTake, const TObjectPtr<ULiveLinkPreset>& InPreset,
+	const FFrameRate* InFrameRate, const FString& InFilePath)
+{
+	FString PackageName = InFilePath;
+	if (PackageName.IsEmpty() && !GetSavePresetPackageName(PackageName))
+	{
 		return;
 	}
@@ -312,8 +324,18 @@
 
 	if (TObjectPtr<ULiveLinkUAssetRecording> NewRecording = NewObject<ULiveLinkUAssetRecording>(NewPackage, *NewAssetName, RF_Public | RF_Standalone))
 	{
-		const double RecordingLength = TimeRecordingEnded - TimeRecordingStarted;
-		NewRecording->InitializeNewRecordingData(MoveTemp(*CurrentRecording), RecordingLength);
+		NewRecording->InitializeNewRecordingData(MoveTemp(InRecordingData), InRecordingLength);
+		// If preset or frame rate provided make sure to set them after initializing new data.
+		if (InPreset)
+		{
+			// Copy the provided preset to our preset. This allows full source names to be preserved if we're saving data from another recording,
+			// such as during an upgrade.
+			UEngine::CopyPropertiesForUnrelatedObjects(InPreset, NewRecording->RecordingPreset);	
+		}
+		if (InFrameRate)
+		{
+			NewRecording->FrameRate = *InFrameRate;
+		}
 
 		NewRecording->MarkPackageDirty();
 
@@ -329,8 +351,11 @@
 	}
 
 	// Increment take on successful recording.
-	ILiveLinkRecordingSessionInfo& SessionInfo = ILiveLinkRecordingSessionInfo::Get();
-	SessionInfo.SetTakeNumber(SessionInfo.GetTakeNumber() + 1);
+	if (bIncrementTake)
+	{
+		ILiveLinkRecordingSessionInfo& SessionInfo = ILiveLinkRecordingSessionInfo::Get();
+		SessionInfo.SetTakeNumber(SessionInfo.GetTakeNumber() + 1);
+	}
 }
 
 void FLiveLinkUAssetRecorder::RecordInitialStaticData()
diff --git a/Engine/Plugins/Animation/LiveLinkHub/Source/LiveLinkHub/Private/Recording/Implementations/LiveLinkUAssetRecorder.h b/Engine/Plugins/Animation/LiveLinkHub/Source/LiveLinkHub/Private/Recording/Implementations/LiveLinkUAssetRecorder.h
--- a/Engine/Plugins/Animation/LiveLinkHub/Source/LiveLinkHub/Private/Recording/Implementations/LiveLinkUAssetRecorder.h
+++ b/Engine/Plugins/Animation/LiveLinkHub/Source/LiveLinkHub/Private/Recording/Implementations/LiveLinkUAssetRecorder.h
@@ -25,15 +25,17 @@ public:
 	virtual void RecordFrameData(const FLiveLinkSubjectKey& SubjectKey, const FLiveLinkFrameDataStruct& FrameData) override;
 	virtual bool IsRecording() const override;
 	virtual bool CanRecord(FText* OutErrorMessage) const override;
+	virtual void SaveFromExistingRecording(const ULiveLinkRecording* InRecording, const FString& InFilePath) override;
+	virtual bool GetSavePresetPackageName(FString& OutName, const FString& InDefaultName = FString(), bool bAlwaysOpenSaveAsDialog = false) const override;
 	//~ End ILiveLinkRecorder
 
 private:
 	/** Prompt the user for a destination path for the recording. */
-	bool OpenSaveDialog(const FString& InDefaultPath, const FString& InNewNameSuggestion, FString& OutPackageName);
-	/** Creates a unique asset name and prompts the user for the recording name. */
-	bool GetSavePresetPackageName(FString& OutName);
+	bool OpenSaveDialog(const FString& InDefaultPath, const FString& InNewNameSuggestion, FString& OutPackageName) const;
 	/** Create a recording package and save it. */
-	void SaveRecording();
+	void SaveRecording(FLiveLinkUAssetRecordingData&& InRecordingData, double InRecordingLength,
+		bool bIncrementTake, const TObjectPtr<ULiveLinkPreset>& InPreset = nullptr, const FFrameRate* InFrameRate = nullptr,
+		const FString& InFilePath = FString());
 	/** Record data to a ULiveLinkRecording object. */
 	void RecordBaseData(FLiveLinkRecordingBaseDataContainer& StaticDataContainer, TSharedPtr<FInstancedStruct>&& DataToRecord);
 	/** Record initial data for all livelink subjects. (Useful when static data was sent before the recording started). */
diff --git a/Engine/Plugins/Animation/LiveLinkHub/Source/LiveLinkHub/Private/Recording/Implementations/LiveLinkUAssetRecording.cpp b/Engine/Plugins/Animation/LiveLinkHub/Source/LiveLinkHub/Private/Recording/Implementations/LiveLinkUAssetRecording.cpp
--- a/Engine/Plugins/Animation/LiveLinkHub/Source/LiveLinkHub/Private/Recording/Implementations/LiveLinkUAssetRecording.cpp
+++ b/Engine/Plugins/Animation/LiveLinkHub/Source/LiveLinkHub/Private/Recording/Implementations/LiveLinkUAssetRecording.cpp
@@ -5,7 +5,6 @@
 #include "Engine/Engine.h"
 #include "Features/IModularFeatures.h"
 #include "HAL/IConsoleManager.h"
-#include "LiveLinkHubClient.h"
 #include "LiveLinkHubLog.h"
 #include "LiveLinkHubModule.h"
 #include "LiveLinkUAssetRecordingPlayer.h"
@@ -14,7 +13,6 @@
 #include "Modules/ModuleManager.h"
 #include "Recording/LiveLinkHubPlaybackController.h"
 #include "Recording/LiveLinkRecordingRangeHelpers.h"
-#include "Serialization/BufferArchive.h"
 #include "Serialization/LargeMemoryReader.h"
 #include "Serialization/ObjectAndNameAsStringProxyArchive.h"
 #include "UObject/Package.h"
@@ -72,6 +70,37 @@ FFrameRate ULiveLinkUAssetRecording::GetGlobalFrameRate() const
 	return CalculateFrameRate(GetMaxFrames(), GetLastTimestamp());
 }
 
+void ULiveLinkUAssetRecording::LoadEntireRecording()
+{
+	auto LoadEntireRecordingLambda = [this]()
+	{
+		const int32 SavedMaxFrames = GetMaxFrames();
+		if (SavedMaxFrames <= 0)
+		{
+			UE_LOG(LogLiveLinkHub, Warning, TEXT("Could not load the recording because the max frame size is 0."));
+			return;
+		}
+		// Start from a clean state.
+		UnloadRecordingData();
+
+		bIsPerformingFullLoad = true;
+
+		LoadRecordingData(0, SavedMaxFrames);	
+	};
+	
+	if (!bPerformedInitialLoad)
+	{
+		GetOnRecordingInitialLoad().BindLambda(LoadEntireRecordingLambda);
+
+		// Loads max frames.
+		LoadRecordingData(0, 1);
+	}
+	else
+	{
+		LoadEntireRecordingLambda();
+	}
+}
+
 void ULiveLinkUAssetRecording::SaveRecordingData()
 {
 	bIsSavingRecordingData = true;
@@ -79,7 +108,7 @@ void ULiveLinkUAssetRecording::SaveRecordingData()
 	TArray64<uint8> Memory;
 	FMemoryWriter64 Archive(Memory);
 	
-	int32 RecordingVersionToSave = RecordingVersion;
+	int32 RecordingVersionToSave = CurrentRecordingVersion;
 	Archive << RecordingVersionToSave;
 
 	// How much static data to expect.
@@ -164,6 +193,12 @@ void ULiveLinkUAssetRecording::LoadRecordingData(int32 InInitialFrame, int32 InN
 	}
 }
 
+void ULiveLinkUAssetRecording::LoadEntireRecordingForUpgrade()
+{
+	bIsUpgrading = true;
+	LoadEntireRecording();
+}
+
 void ULiveLinkUAssetRecording::UnloadRecordingData()
 {
 	// We need to prevent unloading if a package is being saved, but if this is called in a case where there is no outer,
@@ -175,6 +210,8 @@ void ULiveLinkUAssetRecording::UnloadRecordingData()
 		UE_LOG(LogLiveLinkHub, Warning, TEXT("Attempted to unload %s while the package was still being saved"), *GetName());
 		return;
 	}
+
+	bIsPerformingFullLoad = false;
 	
 	bCancelStream = true;
 	UnpauseStream();
@@ -257,7 +294,7 @@ bool ULiveLinkUAssetRecording::WaitForBufferedFrames(int32 InMinFrame, int32 InM
 		InMaxFrame = FMath::Clamp(InMaxFrame, 0, EndFrame);
 
 		const int32 InTotalFrames = InMaxFrame - InMinFrame + 1;
-		const TRange<int32> InRange(InMinFrame, InMaxFrame);
+		const TRange<int32> InRange = UE::LiveLinkHub::RangeHelpers::Private::MakeInclusiveRange(InMinFrame, InMaxFrame);
 
 		while (true)
 		{
@@ -345,6 +382,7 @@ void ULiveLinkUAssetRecording::InitializeNewRecordingData(FLiveLinkUAssetRecordi
 {
 	RecordingData = MoveTemp(InRecordingData);
 	LengthInSeconds = InRecordingLengthSeconds;
+	SavedRecordingVersion = CurrentRecordingVersion;
 	FrameRate = FApp::GetTimecodeFrameRate();
 	RecordingPreset->BuildFromClient();
 
@@ -434,6 +472,9 @@ void ULiveLinkUAssetRecording::SaveFrameData(FArchive* InFileWriter, const FLive
 	
 	SerializedFrameSizes.Reset();
 
+	// Store all timestamps in a block.
+	*InFileWriter << InBaseDataContainer.Timestamps;
+	
 	// Write the frame header size.
 	{
 		const uint64 CurrentPosition = InFileWriter->Tell();
@@ -456,10 +497,6 @@ void ULiveLinkUAssetRecording::SaveFrameData(FArchive* InFileWriter, const FLive
 			
 		// Write the frame index for streaming frames when loading.
 		*InFileWriter << FrameIdx;
-
-		// Write the frame's timestamp.
-		double Timestamp = InBaseDataContainer.Timestamps[FrameIdx];
-		*InFileWriter << Timestamp;
 			
 		// Write the entire frame data.
 		FObjectAndNameAsStringProxyArchive StructAr(*InFileWriter, false);
@@ -509,24 +546,33 @@ void ULiveLinkUAssetRecording::LoadRecordingAsync(int32 InStartFrame, int32 InCu
 	const int32 MaxPossibleFrame = RecordingMaxFrames - 1;
 	InStartFrame = FMath::Clamp(InStartFrame, 0, MaxPossibleFrame);
 	InCurrentFrame = FMath::Clamp(InCurrentFrame, 0, MaxPossibleFrame);
-	const int32 EndFrame = InStartFrame + InNumFramesToLoad - 1;
+	const int32 EndFrame = InStartFrame + InNumFramesToLoad;
 
 	UE_LOG(LogLiveLinkHub, Verbose, TEXT("Loading recording StartFrame: %d, CurrentFrame: %d, EndFrame: %d "), InStartFrame, InCurrentFrame, EndFrame);
 	
 	ON_SCOPE_EXIT
 	{
+		const bool bHasPerformedInitialLoad = bPerformedInitialLoad;
 		// Always set to true. Some blocking operations wait for this, and in the case of a non-fatal error
 		// we want to display error logs and don't want the program to freeze.
 		bPerformedInitialLoad = !bCancelStream;
+
+		if (!bHasPerformedInitialLoad && bPerformedInitialLoad)
+		{
+			FSimpleDelegateGraphTask::CreateAndDispatchWhenReady(
+				FSimpleDelegateGraphTask::FDelegate::CreateLambda([WeakRecording = MakeWeakObjectPtr(this)] ()
+					{
+						if (WeakRecording.IsValid())
+						{
+							WeakRecording->GetOnRecordingInitialLoad().ExecuteIfBound();
+						}
+					}),
+			TStatId(),
+			nullptr,
+			ENamedThreads::GameThread);
+		}
 	};
 	
-	if (IsFrameRangeBuffered(RangeHelpers::Private::MakeInclusiveRange(InStartFrame, FMath::Min(EndFrame, MaxPossibleFrame > 0 ? MaxPossibleFrame : EndFrame))))
-	{
-		// All frames are already buffered.
-		UE_LOG(LogLiveLinkHub, Verbose, TEXT("Frame range already buffered"));
-		return;
-	}
-	
 	// Perform initial load and record entry frame file offsets.
 	const bool bInitialLoad = FrameFileData.Num() == 0;
 	if (bInitialLoad)
@@ -537,9 +583,9 @@ void ULiveLinkUAssetRecording::LoadRecordingAsync(int32 InStartFrame, int32 InCu
 		AnimationData.ReadBulkDataPrimitive(LoadedRecordingVersion);
 
 		// If we modify the RecordingVersion we can perform import logic here.
-		if (LoadedRecordingVersion != RecordingVersion)
+		if (LoadedRecordingVersion != CurrentRecordingVersion)
 		{
-			UE_LOG(LogLiveLinkHub, Log, TEXT("Converting version %d to %d"), LoadedRecordingVersion, RecordingVersion);
+			UE_LOG(LogLiveLinkHub, Log, TEXT("Converting version %d to %d"), LoadedRecordingVersion, CurrentRecordingVersion);
 		}
 			
 		RecordingVersionBeingLoaded = LoadedRecordingVersion;
@@ -586,6 +632,13 @@ void ULiveLinkUAssetRecording::LoadRecordingAsync(int32 InStartFrame, int32 InCu
 			AnimationData.SetBulkDataOffset(EndBlockPosition);
 			FrameFileData.Add(*KeyPosition.FrameDataSubjectKey, MoveTemp(KeyPosition));
 		}
+
+		if (bIsPerformingFullLoad)
+		{
+			// Normally set at the end of scope, but if doing a full load, we won't hit that until after everything is loaded,
+			// and we may have observers who want to know the status.
+			bPerformedInitialLoad = true;
+		}
 	}
 	
 	// Load the required frames, either on initial load or subsequent loads.
@@ -739,24 +792,36 @@ bool ULiveLinkUAssetRecording::LoadInitialFrameData(UE::LiveLinkHub::FrameData::
 				MaxFrameDiskSize = MaxValue;
 			}
 		}
-	}
 
-	// Find the last timestamp, this is so we can calculate the correct framerate for this track.
-	{
-		const int32 LastFrame = OutFrameData.MaxFrames - 1;
-		double Timestamp;
-		if (!LoadTimestampFromDisk(LastFrame, OutFrameData, Timestamp))
+		const bool bUsesOldTimestamps = RecordingVersionBeingLoaded < UE::LiveLinkHub::Private::RecordingVersions::BulkTimestamps;
+		// Find the last timestamp, this is so we can calculate the correct framerate for this track.
+		if (bUsesOldTimestamps)
 		{
-			return false;
-		}
-		OutFrameData.LastTimestamp = Timestamp;
+			OutFrameData.Timestamps.Reserve(OutFrameData.MaxFrames);
+		
+			for (int32 Idx = 0; Idx < OutFrameData.MaxFrames; ++Idx)
+			{
+				double Timestamp = 0.0;
+				if (!LoadTimestampFromDisk(Idx, OutFrameData, Timestamp))
+				{
+					return false;
+				}
 
-		if (Timestamp > RecordingLastTimestamp)
+				OutFrameData.Timestamps.Add(Timestamp);
+			}
+		}
+		else
 		{
-			RecordingLastTimestamp = Timestamp;
+			Reader->GetMemoryReader() << OutFrameData.Timestamps;
 		}
 	}
 
+	OutFrameData.LastTimestamp = OutFrameData.Timestamps.Num () > 0 ? OutFrameData.Timestamps.Last() : 0.0;
+	if (OutFrameData.LastTimestamp > RecordingLastTimestamp)
+	{
+		RecordingLastTimestamp = OutFrameData.LastTimestamp;
+	}
+
 	// Calculate frame rate.
 	OutFrameData.LocalFrameRate = CalculateFrameRate(OutFrameData.MaxFrames, OutFrameData.LastTimestamp);
 	
@@ -771,17 +836,45 @@ void ULiveLinkUAssetRecording::LoadFrameData(UE::LiveLinkHub::FrameData::Private
 	using namespace UE::LiveLinkHub;
 	
 	FrameData::Private::FFrameBufferIterationData& IterationData = InFrameData.BufferIterationData;
-
+	
+	const FFrameRate RecordingFrameRate = GetGlobalFrameRate();
+	const FQualifiedFrameTime StartFrameTime(FFrameTime(RequestedStartFrame), RecordingFrameRate);
+	const FQualifiedFrameTime StartInitialFrameTime(FFrameTime(RequestedInitialFrame), RecordingFrameRate);
+	
+	double RequestedStartTimestamp = StartFrameTime.AsSeconds();
+	double RequestedInitialTimestamp = StartInitialFrameTime.AsSeconds();
+	if (FMath::IsNaN(RequestedStartTimestamp))
+	{
+		RequestedStartTimestamp = 0.0;
+	}
+	if (FMath::IsNaN(RequestedInitialTimestamp))
+	{
+		RequestedInitialTimestamp = 0.0;
+	}
+	
 	// First, localize the frame times to this frame data.
 	if (InFrameData.LocalFrameRate.IsValid())
 	{
-		const FFrameRate RecordingFrameRate = GetGlobalFrameRate();
-		const FQualifiedFrameTime StartFrameTime(FFrameTime(RequestedStartFrame), RecordingFrameRate);
-		const FQualifiedFrameTime StartInitialFrameTime(FFrameTime(RequestedInitialFrame), RecordingFrameRate);
-
 		RequestedStartFrame = StartFrameTime.ConvertTo(InFrameData.LocalFrameRate).FrameNumber.Value;
 		RequestedInitialFrame = StartInitialFrameTime.ConvertTo(InFrameData.LocalFrameRate).FrameNumber.Value;
 	}
+
+	// Localizing the frame rates may not be enough if the track had inconsistent frames or frame rate. Next, we
+	// need to properly identify the exact local frame to use most closely matching the requested timestamp.
+	// This should produce the same index unless the track had an inconsistent storage and frame rate.
+	// E.g., times 0-1:00 had data, 1:00-1:30 had no data, and 1:30 to 2:00 had data again.
+	int32 StartInitialFrameLocalOffset = 0;
+	{
+		const int32 GlobalRequestedInitialFrame = RequestedInitialFrame;
+	
+		RequestedStartFrame = InFrameData.FindFrameFromTimestamp(RequestedStartTimestamp, RequestedStartFrame);
+		RequestedInitialFrame = InFrameData.FindFrameFromTimestamp(RequestedInitialTimestamp, RequestedInitialFrame);
+
+		// Store any offset from the calculated localized frame to the actual localized frame.
+		// This is needed to report the global buffered frames correctly. As a whole, the recording cares whether the
+		// requested frames should be considered loaded, not the local frames of each track.
+		StartInitialFrameLocalOffset = GlobalRequestedInitialFrame - RequestedInitialFrame;
+	}
 	
 	int32 MaxFrames = InFrameData.MaxFrames;
 	if (MaxFrames > 0)
@@ -789,7 +882,7 @@ void ULiveLinkUAssetRecording::LoadFrameData(UE::LiveLinkHub::FrameData::Private
 		if (RequestedFramesToLoad > 0)
 		{
 			// Don't go past requested frames or max frames.
-			MaxFrames = FMath::Min(MaxFrames, RequestedStartFrame + RequestedFramesToLoad);
+			MaxFrames = FMath::Min(MaxFrames, RequestedStartFrame + RequestedFramesToLoad + 1);
 		}
 
 		const int32 LastFrame = MaxFrames - 1;
@@ -857,7 +950,7 @@ void ULiveLinkUAssetRecording::LoadFrameData(UE::LiveLinkHub::FrameData::Private
 			}
 
 			LoadedRange = InDataContainer.GetBufferedFrames();
-			if (!LoadedRange.Contains(InitialFrameToStream))
+			if (!LoadedRange.Contains(RequestedInitialFrame))
 			{
 				// In this case we aren't immediately loading frames, but will get to them eventually.
 				TRange<int32> Intersection = TRange<int32>::Intersection(RangeToLoad, LoadedRange);
@@ -1097,7 +1190,7 @@ void ULiveLinkUAssetRecording::LoadFrameData(UE::LiveLinkHub::FrameData::Private
 			{
 				// Frame needs to be deserialized, and potentially loaded from disk.
 				
-				double Timestamp;
+				double Timestamp = 0.0;
 				TSharedPtr<FInstancedStruct> DataInstancedStruct;
 				TSharedPtr<FLiveLinkHubBulkData::FScopedBulkDataMemoryReader> MemoryToUse;
 				// Check if we have preloaded the raw data into memory.
@@ -1145,6 +1238,8 @@ void ULiveLinkUAssetRecording::LoadFrameData(UE::LiveLinkHub::FrameData::Private
 		IterationData.LastLoadedLeftFrame = LastLoadedLeftFrame;
 		IterationData.LastLoadedRightFrame = LastLoadedRightFrame;
 
+		InDataContainer.LocalFrameOffset = StartInitialFrameLocalOffset;
+		
 		// Determine the iteration status, which will impact the overall async load loop.
 		if (RightFrameIdx >= MaxFrames && LeftFrameIdx <= RequestedStartFrame)
 		{
@@ -1157,8 +1252,10 @@ void ULiveLinkUAssetRecording::LoadFrameData(UE::LiveLinkHub::FrameData::Private
 	}
 }
 
-bool ULiveLinkUAssetRecording::LoadFrameFromDisk(const int32 InFrame, const UE::LiveLinkHub::FrameData::Private::FFrameMetaData& InFrameData,
-	TSharedPtr<FInstancedStruct>& OutFrame, double& OutTimestamp, const TSharedPtr<FLiveLinkHubBulkData::FScopedBulkDataMemoryReader>& InMemory)
+bool ULiveLinkUAssetRecording::LoadFrameFromDisk(const int32 InFrame,
+	const UE::LiveLinkHub::FrameData::Private::FFrameMetaData& InFrameData,
+	TSharedPtr<FInstancedStruct>& OutFrame, double& OutTimestamp,
+	const TSharedPtr<FLiveLinkHubBulkData::FScopedBulkDataMemoryReader>& InMemory)
 {
 	DECLARE_SCOPE_CYCLE_COUNTER(TEXT("ULiveLinkUAssetRecording::LoadFrameFromDisk"), STAT_ULiveLinkUAssetRecording_LoadFrameFromDisk, STATGROUP_LiveLinkHub);
 	OutTimestamp = 0.0;
@@ -1187,8 +1284,16 @@ bool ULiveLinkUAssetRecording::LoadFrameFromDisk(const int32 InFrame, const UE::
 		return false;
 	}
 			
-	double Timestamp = 0;
-	Reader->GetMemoryReader() << Timestamp;
+	double Timestamp = 0.0;
+	if (RecordingVersionBeingLoaded < UE::LiveLinkHub::Private::RecordingVersions::BulkTimestamps)
+	{
+		Reader->GetMemoryReader() << Timestamp;
+	}
+	else
+	{
+		check(InFrameData.Timestamps.IsValidIndex(InFrame));
+		Timestamp = InFrameData.Timestamps[InFrame];
+	}
 
 	// Instantiate the animation frame.
 	UScriptStruct* LoadedStruct = InFrameData.LoadedStruct.Get();
@@ -1213,6 +1318,8 @@ bool ULiveLinkUAssetRecording::LoadFrameFromDisk(const int32 InFrame, const UE::
 
 bool ULiveLinkUAssetRecording::LoadTimestampFromDisk(const int32 InFrame, const UE::LiveLinkHub::FrameData::Private::FFrameMetaData& InFrameData, double& OutTimestamp)
 {
+	check(RecordingVersionBeingLoaded < UE::LiveLinkHub::Private::RecordingVersions::BulkTimestamps);
+	
 	const int64 FramePosition = InFrameData.GetFrameFilePosition(InFrame);
 	AnimationData.SetBulkDataOffset(FramePosition);
 			
@@ -1323,6 +1430,7 @@ void ULiveLinkUAssetRecording::MoveRangeToCache(const TRange<int32>& InRange, FL
 			HistoryContainer.Timestamps.Append(&InDataContainer.Timestamps[StartIndex], CountToRemove);
 			HistoryContainer.RecordedData.Append(&InDataContainer.RecordedData[StartIndex], CountToRemove);
 			HistoryContainer.RecordedDataStartFrame = InDataContainer.RecordedDataStartFrame + StartIndex;
+			HistoryContainer.LocalFrameOffset = InDataContainer.LocalFrameOffset;
 
 			InFrameData.BufferedCache.FrameData.Add(MoveTemp(HistoryContainer));
 
@@ -1362,13 +1470,13 @@ void ULiveLinkUAssetRecording::UpdateBufferedFrames()
 	for (TTuple<FLiveLinkSubjectKey, FrameData::Private::FFrameMetaData>& FrameDataKeyVal : FrameFileData)
 	{
 		FrameData::Private::FFrameMetaData& FrameData = FrameDataKeyVal.Value;
-		RangeHelpers::Private::TRangeArray<int32> CacheBufferRanges = FrameData.BufferedCache.GetCacheBufferRanges();
+		RangeHelpers::Private::TRangeArray<int32> CacheBufferRanges = FrameData.BufferedCache.GetCacheBufferRanges(/*bIncludeOffset*/ true);
 		for (TRange<int32>& Range : CacheBufferRanges)
 		{
-			Range = RangeHelpers::Private::ConvertRangeFrameRate(Range, FrameData.LocalFrameRate, GlobalFrameRate);
+			Range = RangeHelpers::Private::ConvertRangeFrameRate(Range,FrameData.LocalFrameRate, GlobalFrameRate);
 		}
 		
-		TRange<int32> ActiveBufferFrames = RecordingData.FrameData.FindChecked(*FrameData.FrameDataSubjectKey.Get()).GetBufferedFrames();
+		TRange<int32> ActiveBufferFrames = RecordingData.FrameData.FindChecked(*FrameData.FrameDataSubjectKey.Get()).GetBufferedFrames(/*bIncludeOffset*/ true);
 		ActiveBufferFrames = RangeHelpers::Private::ConvertRangeFrameRate(ActiveBufferFrames, FrameData.LocalFrameRate, GlobalFrameRate);
 		
 		BufferedFrameRanges.Append(MoveTemp(CacheBufferRanges));
@@ -1432,6 +1540,25 @@ void ULiveLinkUAssetRecording::FLiveLinkStreamAsyncTask::DoWork()
 			LastInitialFrame = LiveLinkRecording->InitialFrameToStream;
 			LiveLinkRecording->LoadRecordingAsync(LiveLinkRecording->EarliestFrameToStream,
 				LiveLinkRecording->InitialFrameToStream, LiveLinkRecording->TotalFramesToStream);
+
+			if (LiveLinkRecording->bIsPerformingFullLoad)
+			{
+				// Execute on game thread since we're firing a delegate.
+				FSimpleDelegateGraphTask::CreateAndDispatchWhenReady(
+					FSimpleDelegateGraphTask::FDelegate::CreateLambda([WeakRecording = MakeWeakObjectPtr(LiveLinkRecording)] ()
+						{
+							if (WeakRecording.IsValid())
+							{
+								WeakRecording->bIsFullyLoaded = true;
+								WeakRecording->bIsPerformingFullLoad = false;
+								WeakRecording->bIsUpgrading = false;
+								WeakRecording->GetOnRecordingFullyLoaded().ExecuteIfBound();
+							}
+						}),
+				TStatId(),
+				nullptr,
+				ENamedThreads::GameThread);
+			}
 		}
 		else
 		{
diff --git a/Engine/Plugins/Animation/LiveLinkHub/Source/LiveLinkHub/Private/Recording/Implementations/LiveLinkUAssetRecording.h b/Engine/Plugins/Animation/LiveLinkHub/Source/LiveLinkHub/Private/Recording/Implementations/LiveLinkUAssetRecording.h
--- a/Engine/Plugins/Animation/LiveLinkHub/Source/LiveLinkHub/Private/Recording/Implementations/LiveLinkUAssetRecording.h
+++ b/Engine/Plugins/Animation/LiveLinkHub/Source/LiveLinkHub/Private/Recording/Implementations/LiveLinkUAssetRecording.h
@@ -6,6 +6,7 @@
 #include "Data/LiveLinkHubBulkData.h"
 #include "LiveLinkFrameData.h"
 #include "LiveLinkTypes.h"
+#include "Recording/LiveLinkHubRecordingVersions.h"
 #include "Recording/LiveLinkRecording.h"
 #include "StructUtils/InstancedStruct.h"
 #include "UObject/Object.h"
@@ -14,13 +15,6 @@
 
 struct FLiveLinkPlaybackTracks;
 
-namespace UE::LiveLinkHub::Private::RecordingVersions
-{
-	constexpr int32 InitialVersion = 1;
-	/** Account for varying frame sizes and offsets. */
-	constexpr int32 DynamicFrameSizes = 2;
-}
-
 /**
  * Asset containing all animation data stored as bulk data. This is loaded async in chunks dependent on the playhead position.
  * Overall recording length, framerate, and frame indices are based on the maximum track length and farthest timestamp.
@@ -45,6 +39,8 @@ public:
 	virtual int32 GetMaxFrames() const override { return RecordingMaxFrames; }
 	virtual double GetLastTimestamp() const override { return RecordingLastTimestamp; }
 	virtual FFrameRate GetGlobalFrameRate() const override;
+	virtual void LoadEntireRecording() override;
+	virtual bool IsLoadingEntireRecording() const override { return bIsPerformingFullLoad;}
 	
 	/** Save recording data to disk. */
 	void SaveRecordingData();
@@ -52,6 +48,9 @@ public:
 	/** Load recording data from disk. */
 	void LoadRecordingData(int32 InInitialFrame, int32 InNumFramesToLoad);
 
+	/** Load the entire recording and notify this is for an upgrade. */
+	void LoadEntireRecordingForUpgrade();
+	
 	/** Free memory and close file reader. */
 	void UnloadRecordingData();
 
@@ -67,14 +66,28 @@ public:
 	/** Checks if a specific frame range is buffered. */
 	bool IsFrameRangeBuffered(const TRange<int32>& InRange) const;
 
+	/** If the recording is currently upgrading. */
+	bool IsUpgrading() const { return bIsUpgrading; }
+	
 	/** Copy the asset's loaded recording data to a format suitable for playback in live link. */
 	void CopyRecordingData(FLiveLinkPlaybackTracks& InOutLiveLinkPlaybackTracks) const;
 
 	/** Initial setup of new recording data. */
 	void InitializeNewRecordingData(FLiveLinkUAssetRecordingData&& InRecordingData, double InRecordingLengthSeconds);
 
+	/** If the recording has performed its initial load. */
+	bool HasPerformedInitialLoad() const { return bPerformedInitialLoad; }
+	
 	/** Calculate the framerate. */
 	static FFrameRate CalculateFrameRate(const int32 InMaxFrames, const double InTime);
+
+	DECLARE_DELEGATE(FOnRecordingInitialLoad);
+	/** Delegate fired when the recording has performed an initial load. */
+	FOnRecordingInitialLoad& GetOnRecordingInitialLoad() { return OnRecordingInitialLoadEvent; }
+	
+	DECLARE_DELEGATE(FOnRecordingFullyLoaded);
+	/** Delegate fired when a recording has fully loaded. */
+	FOnRecordingFullyLoaded& GetOnRecordingFullyLoaded() { return OnRecordingFullyLoadedEvent; }
 	
 private:
 	/** Serialize the number of frames (array size) of the BaseDataContainer to the archive. */
@@ -108,8 +121,10 @@ private:
 	 * @param OutTimestamp The deserialized timestamp.
 	 * @param InMemory [Optional] If we are reading from memory rather than bulk data directly.
 	 */
-	bool LoadFrameFromDisk(const int32 InFrame, const UE::LiveLinkHub::FrameData::Private::FFrameMetaData& InFrameData,
-		TSharedPtr<FInstancedStruct>& OutFrame, double& OutTimestamp, const TSharedPtr<FLiveLinkHubBulkData::FScopedBulkDataMemoryReader>& InMemory = nullptr);
+	bool LoadFrameFromDisk(const int32 InFrame,
+		const UE::LiveLinkHub::FrameData::Private::FFrameMetaData& InFrameData,
+		TSharedPtr<FInstancedStruct>& OutFrame, double& OutTimestamp,
+		const TSharedPtr<FLiveLinkHubBulkData::FScopedBulkDataMemoryReader>& InMemory = nullptr);
 
 	/** Only load the timestamp from disk. */
 	bool LoadTimestampFromDisk(const int32 InFrame, const UE::LiveLinkHub::FrameData::Private::FFrameMetaData& InFrameData, double& OutTimestamp);
@@ -149,7 +164,6 @@ private:
 
 	/** Called after garbage collection. */
 	void OnPostGarbageCollect();
-	
 public:
 	/** Recorded static and frame data. */
 	UPROPERTY()
@@ -223,6 +237,9 @@ private:
 	/** If we are currently saving recording frame data to disk. */
 	std::atomic<bool> bIsSavingRecordingData = false;
 
+	/** If we are currently loading all data from disk. */
+	std::atomic<bool> bIsPerformingFullLoad = false;
+
 	/** The maximum frame disk size across frame data. */
 	std::atomic<int32> MaxFrameDiskSize = 0;
 
@@ -249,9 +266,18 @@ private:
 
 	/** If the recording is fully loaded into memory. */
 	bool bIsFullyLoaded = false;
+
+	/** Upgrade in progress. */
+	bool bIsUpgrading = false;
+
+	/** Called on the main thread when the recording has performed an initial load. */
+	FOnRecordingFullyLoaded OnRecordingInitialLoadEvent;
+	
+	/** Called on the main thread when the recording has fully loaded. */
+	FOnRecordingFullyLoaded OnRecordingFullyLoadedEvent;
 	
 	/** The current version of the recording. */
-	const int32 RecordingVersion = UE::LiveLinkHub::Private::RecordingVersions::DynamicFrameSizes;
+	const int32 CurrentRecordingVersion = UE::LiveLinkHub::Private::RecordingVersions::Latest;
 
 	/** The version being currently loaded. */
 	int32 RecordingVersionBeingLoaded = 0;
diff --git a/Engine/Plugins/Animation/LiveLinkHub/Source/LiveLinkHub/Private/Recording/LiveLinkHubPlaybackController.cpp b/Engine/Plugins/Animation/LiveLinkHub/Source/LiveLinkHub/Private/Recording/LiveLinkHubPlaybackController.cpp
--- a/Engine/Plugins/Animation/LiveLinkHub/Source/LiveLinkHub/Private/Recording/LiveLinkHubPlaybackController.cpp
+++ b/Engine/Plugins/Animation/LiveLinkHub/Source/LiveLinkHub/Private/Recording/LiveLinkHubPlaybackController.cpp
@@ -97,6 +97,30 @@ TSharedRef<SWidget> FLiveLinkHubPlaybackController::MakePlaybackWidget()
 	return SNew(SLiveLinkHubPlaybackWidget)
 		.Visibility_Lambda([this]() { return IsInPlayback() ? EVisibility::Visible : EVisibility::Collapsed; })
 		.IsPlaybackEnabled_Raw(this, &FLiveLinkHubPlaybackController::IsReady)
+		.IsUpgrading_Lambda([this]()
+		{
+			if (const ULiveLinkUAssetRecording* Recording = Cast<ULiveLinkUAssetRecording>(GetRecording().Get()))
+			{
+				return Recording->IsUpgrading() && Recording->IsLoadingEntireRecording();
+			}
+			return false;
+		})
+		.OnGetUpgradePercent_Lambda([this]()
+		{
+			if (const ULiveLinkUAssetRecording* Recording = Cast<ULiveLinkUAssetRecording>(GetRecording().Get()))
+			{
+				using namespace UE::LiveLinkHub::RangeHelpers::Private;
+				const TRangeArray<int32> BufferedFrames = GetBufferedFrameRanges();
+				const float Length = GetRangeLength(BufferedFrames);
+				const float MaxLength = RecordingToPlay->GetMaxFrames();
+				// Wait for initial load otherwise the value can jump since we do a partial load, then unload first.
+				if (MaxLength > 0.f && Recording->HasPerformedInitialLoad())
+				{
+					return (Length / MaxLength) * 100.f;
+				}
+			}
+			return 0.f;
+		})
 		.OnPlayForward_Raw(this, &FLiveLinkHubPlaybackController::BeginPlayback, false)
 		.OnPlayReverse_Raw(this, &FLiveLinkHubPlaybackController::BeginPlayback, true)
 		.OnFirstFrame_Lambda([this]()
@@ -228,11 +252,23 @@ void FLiveLinkHubPlaybackController::PreparePlayback(ULiveLinkRecording* InLiveL
 					LiveLinkClient.RemoveAllSources();
 					LiveLinkClient.Tick();
 				}
-		
+
 				RecordingToPlay->RecordingPreset->ApplyToClientLatent([this](bool)
 				{
 					bIsReady = true;
 					SyncToFrame(FQualifiedFrameTime()); // Needed to establish connection with client
+					
+					// We should already be on the game thread, but we want to execute OnPlaybackReady delegate on the next tick
+					// since we're still operating within a callback of ApplyToClient. If a caller performs an action
+					// that prompts another ApplyToClientCall while inside one, we can hit an exception.
+					FSimpleDelegateGraphTask::CreateAndDispatchWhenReady(
+						FSimpleDelegateGraphTask::FDelegate::CreateLambda([this]()
+						{
+							OnPlaybackReady().Broadcast();
+						}),
+						TStatId(),
+						nullptr,
+						ENamedThreads::GameThread);
 				});
 			}
 		};
@@ -247,6 +283,11 @@ void FLiveLinkHubPlaybackController::PreparePlayback(ULiveLinkRecording* InLiveL
 			PreparePlaybackCallback();
 		}
 	}
+	else
+	{
+		// Already loaded and ready.
+		OnPlaybackReady().Broadcast();
+	}
 }
 
 void FLiveLinkHubPlaybackController::PlayRecording(ULiveLinkRecording* InLiveLinkRecording)
@@ -366,7 +407,8 @@ void FLiveLinkHubPlaybackController::Eject(TFunction<void()> CompletionCallback)
 	}
 }
 
-void FLiveLinkHubPlaybackController::EjectAndUnload(TFunction<void()> EjectCompletionCallback, const ULiveLinkRecording* InRecording, bool bClosePlaybackTab)
+void FLiveLinkHubPlaybackController::EjectAndUnload(const TFunction<void()>& EjectCompletionCallback,
+	const ULiveLinkRecording* InRecording, bool bClosePlaybackTab)
 {
 	const ULiveLinkRecording* Recording = InRecording ? InRecording : RecordingToPlay.Get();
 	UPackage* Package = Recording ? Recording->GetPackage() : nullptr;
diff --git a/Engine/Plugins/Animation/LiveLinkHub/Source/LiveLinkHub/Private/Recording/LiveLinkHubPlaybackController.h b/Engine/Plugins/Animation/LiveLinkHub/Source/LiveLinkHub/Private/Recording/LiveLinkHubPlaybackController.h
--- a/Engine/Plugins/Animation/LiveLinkHub/Source/LiveLinkHub/Private/Recording/LiveLinkHubPlaybackController.h
+++ b/Engine/Plugins/Animation/LiveLinkHub/Source/LiveLinkHub/Private/Recording/LiveLinkHubPlaybackController.h
@@ -14,15 +14,14 @@
 
 #include <atomic>
 
+#include "LiveLinkRecording.h"
 #include "LiveLinkRecordingRangeHelpers.h"
 
 enum class ELiveLinkHubPlaybackState : uint8;
 class ILiveLinkClient;
 struct FLiveLinkRecordingData;
 class FRunnableThread;
-class ULiveLinkRecording;
 class ULiveLinkPreset;
-class ULiveLinkRecording;
 class SWidget;
 class FLiveLinkHubAtomicQualifiedFrameTime;
 
@@ -64,7 +63,8 @@ public:
 	 * @param EjectCompletionCallback Callback after ejection occurs.
 	 * @param InRecording The specific recording to eject and unload. If null, the RecordingToPlay is used. 
 	 */
-	void EjectAndUnload(TFunction<void()> EjectCompletionCallback = nullptr, const ULiveLinkRecording* InRecording = nullptr, bool bClosePlaybackTab = true);
+	void EjectAndUnload(const TFunction<void()>& EjectCompletionCallback = nullptr,
+		const ULiveLinkRecording* InRecording = nullptr, bool bClosePlaybackTab = true);
 
 	/**
 	 * Completely unload a recording package.
@@ -111,7 +111,7 @@ public:
 	/** If the controller is ready for commands. */
 	bool IsReady() const
 	{
-		return bIsReady;
+		return bIsReady && GetRecording() && !GetRecording()->IsLoadingEntireRecording();
 	}
 	
 	/** If a recording is loaded into the controller. */
@@ -158,6 +158,12 @@ public:
 	{
 		return PlaybackFinishedDelegate;
 	}
+
+	/** Delegate called when recording playback is ready. */
+	FTSSimpleMulticastDelegate& OnPlaybackReady()
+	{
+		return PlaybackReadyDelegate;
+	}
 	
 	/** Create the playback thread. */
 	void Start();
@@ -231,6 +237,8 @@ private:
 	TStrongObjectPtr<ULiveLinkRecording> RecordingToPlay;
 	/** Delegate called when a recording playback is finished (if it's not looping). */
 	FTSSimpleMulticastDelegate PlaybackFinishedDelegate;
+	/** Delegate called when recording playback is ready. */
+	FTSSimpleMulticastDelegate PlaybackReadyDelegate;
 	/** Preset used to rollback the hub to its previous state after playing a recording. */
 	TStrongObjectPtr<ULiveLinkPreset> RollbackPreset;
 	/** Atomic bool keeping track of whether we should loop the playback. */
diff --git a/Engine/Plugins/Animation/LiveLinkHub/Source/LiveLinkHub/Private/Recording/LiveLinkHubRecordingController.cpp b/Engine/Plugins/Animation/LiveLinkHub/Source/LiveLinkHub/Private/Recording/LiveLinkHubRecordingController.cpp
--- a/Engine/Plugins/Animation/LiveLinkHub/Source/LiveLinkHub/Private/Recording/LiveLinkHubRecordingController.cpp
+++ b/Engine/Plugins/Animation/LiveLinkHub/Source/LiveLinkHub/Private/Recording/LiveLinkHubRecordingController.cpp
@@ -73,4 +73,15 @@ void FLiveLinkHubRecordingController::RecordFrameData(const FLiveLinkSubjectKey&
 	RecorderImplementation->RecordFrameData(SubjectKey, FrameData);
 }
 
+void FLiveLinkHubRecordingController::RecordFromExistingAndSave(const ULiveLinkUAssetRecording* InRecording, const FString& InFilePath)
+{
+	RecorderImplementation->SaveFromExistingRecording(InRecording, InFilePath);
+}
+
+bool FLiveLinkHubRecordingController::OpenSaveAsPrompt(FString& OutName, const FString& InInitialName) const
+{
+	constexpr bool bOpenSaveAs = true;
+	return RecorderImplementation->GetSavePresetPackageName(OutName, InInitialName, bOpenSaveAs);
+}
+
 #undef LOCTEXT_NAMESPACE
diff --git a/Engine/Plugins/Animation/LiveLinkHub/Source/LiveLinkHub/Private/Recording/LiveLinkHubRecordingController.h b/Engine/Plugins/Animation/LiveLinkHub/Source/LiveLinkHub/Private/Recording/LiveLinkHubRecordingController.h
--- a/Engine/Plugins/Animation/LiveLinkHub/Source/LiveLinkHub/Private/Recording/LiveLinkHubRecordingController.h
+++ b/Engine/Plugins/Animation/LiveLinkHub/Source/LiveLinkHub/Private/Recording/LiveLinkHubRecordingController.h
@@ -6,6 +6,7 @@
 #include "Templates/PimplPtr.h"
 
 class ILiveLinkRecorder;
+class ULiveLinkUAssetRecording;
 struct FInstancedStruct;
 struct FLiveLinkRecordingBaseDataContainer;
 class SWidget;
@@ -38,6 +39,12 @@ public:
 
 	/** Record frame data in the current recording. */
 	void RecordFrameData(const FLiveLinkSubjectKey& SubjectKey, const FLiveLinkFrameDataStruct& FrameData);
+
+	/** Record from an existing recording. */
+	void RecordFromExistingAndSave(const ULiveLinkUAssetRecording* InRecording, const FString& InFilePath);
+
+	/** Open a save as window and retrieve a file path. */
+	bool OpenSaveAsPrompt(FString& OutName, const FString& InInitialName) const;
 private:
 	/** Recorder used to serialize livelink data into a given format. */
 	TSharedPtr<ILiveLinkRecorder> RecorderImplementation;
diff --git a/Engine/Plugins/Animation/LiveLinkHub/Source/LiveLinkHub/Private/Recording/LiveLinkHubRecordingVersions.h b/Engine/Plugins/Animation/LiveLinkHub/Source/LiveLinkHub/Private/Recording/LiveLinkHubRecordingVersions.h
new file mode 100644
--- /dev/null
+++ b/Engine/Plugins/Animation/LiveLinkHub/Source/LiveLinkHub/Private/Recording/LiveLinkHubRecordingVersions.h
@@ -0,0 +1,17 @@
+// Copyright Epic Games, Inc. All Rights Reserved.
+
+#pragma once
+
+#include "HAL/Platform.h"
+
+namespace UE::LiveLinkHub::Private::RecordingVersions
+{
+	constexpr int32 InitialVersion = 1;
+	/** Account for varying frame sizes and offsets. */
+	constexpr int32 DynamicFrameSizes = 2;
+	/** Timestamps serialized in bulk. */
+	constexpr int32 BulkTimestamps = 3;
+
+	/** The latest recording version. UPDATE IF ADDING NEW VERSIONS! */
+	constexpr int32 Latest = BulkTimestamps;
+}
diff --git a/Engine/Plugins/Animation/LiveLinkHub/Source/LiveLinkHub/Private/Recording/LiveLinkRecorder.h b/Engine/Plugins/Animation/LiveLinkHub/Source/LiveLinkHub/Private/Recording/LiveLinkRecorder.h
--- a/Engine/Plugins/Animation/LiveLinkHub/Source/LiveLinkHub/Private/Recording/LiveLinkRecorder.h
+++ b/Engine/Plugins/Animation/LiveLinkHub/Source/LiveLinkHub/Private/Recording/LiveLinkRecorder.h
@@ -6,6 +6,7 @@
 #include "LiveLinkTypes.h"
 
 class ULiveLinkRecording;
+struct FLiveLinkUAssetRecordingData;
 
 /** Interface for handling serialization of livelink data. */
 class ILiveLinkRecorder
@@ -25,4 +26,8 @@ public:
 	virtual bool IsRecording() const = 0;
 	/** Returns whether we can begin recording. */
 	virtual bool CanRecord(FText* OutErrorMessage = nullptr) const = 0;
+	/** Directly provide the recording and save it. */
+	virtual void SaveFromExistingRecording(const ULiveLinkRecording* InRecording, const FString& InFilePath) = 0;
+	/** Creates a unique asset name and prompts the user for the recording name. */
+	virtual bool GetSavePresetPackageName(FString& OutName, const FString& InDefaultName, bool bAlwaysOpenSaveAsDialog = false) const = 0;
 };
diff --git a/Engine/Plugins/Animation/LiveLinkHub/Source/LiveLinkHub/Private/Recording/LiveLinkRecording.h b/Engine/Plugins/Animation/LiveLinkHub/Source/LiveLinkHub/Private/Recording/LiveLinkRecording.h
--- a/Engine/Plugins/Animation/LiveLinkHub/Source/LiveLinkHub/Private/Recording/LiveLinkRecording.h
+++ b/Engine/Plugins/Animation/LiveLinkHub/Source/LiveLinkHub/Private/Recording/LiveLinkRecording.h
@@ -34,6 +34,12 @@ public:
 	/** Retrieve the generalized framerate of this recording. */
 	virtual FFrameRate GetGlobalFrameRate() const { return FFrameRate(); }
 
+	/** Load the entire recording into memory. */
+	virtual void LoadEntireRecording() { };
+
+	/** If the entire recording is in the process of being loaded into memory. */
+	virtual bool IsLoadingEntireRecording() const { return false; }
+
 	/** LiveLink Preset used to save the initial state of the sources and subjects at the time of recording. */
 	UPROPERTY(Instanced)
 	TObjectPtr<ULiveLinkPreset> RecordingPreset = nullptr;
@@ -45,4 +51,8 @@ public:
 	/** The framerate of the recording. */
 	UPROPERTY()
 	FFrameRate FrameRate;
+
+	/** The saved version of this recording. */
+	UPROPERTY(AssetRegistrySearchable)
+	int32 SavedRecordingVersion;
 };
\ No newline at end of file
diff --git a/Engine/Plugins/Animation/LiveLinkHub/Source/LiveLinkHub/Private/Recording/LiveLinkRecordingRangeHelpers.h b/Engine/Plugins/Animation/LiveLinkHub/Source/LiveLinkHub/Private/Recording/LiveLinkRecordingRangeHelpers.h
--- a/Engine/Plugins/Animation/LiveLinkHub/Source/LiveLinkHub/Private/Recording/LiveLinkRecordingRangeHelpers.h
+++ b/Engine/Plugins/Animation/LiveLinkHub/Source/LiveLinkHub/Private/Recording/LiveLinkRecordingRangeHelpers.h
@@ -27,6 +27,29 @@ namespace UE::LiveLinkHub::RangeHelpers::Private
 		return UpperBound - LowerBound + 1;
 	}
 
+	/** Gets the length of the range array. Does not account for inclusive ranges. */
+	template<typename T>
+	T GetRangeLength(const TRangeArray<T>& InRangeArray)
+	{
+		T MinLowerBound = TNumericLimits<T>::Max();
+		T MaxUpperBound = TNumericLimits<T>::Min();
+		for (const TRange<T>& InRange : InRangeArray)
+		{
+			const T LowerBound = InRange.GetLowerBoundValue();
+			const T UpperBound = InRange.GetUpperBoundValue();
+			if (LowerBound < MinLowerBound)
+			{
+				MinLowerBound = LowerBound;
+			}
+			if (UpperBound > MaxUpperBound)
+			{
+				MaxUpperBound = UpperBound;
+			}
+		}
+
+		return MaxUpperBound - MinLowerBound + 1;
+	}
+	
 	/** Makes an inclusive range. */
 	template<typename T>
 	TRange<T> MakeInclusiveRange(const T InStart, const T InEnd)
@@ -44,6 +67,6 @@ namespace UE::LiveLinkHub::RangeHelpers::Private
 		const FFrameTime ToStart = FromStart.ConvertTo(InToFrameRate);
 		const FFrameTime ToEnd = FromEnd.ConvertTo(InToFrameRate);
 
-		return TRange<T>(ToStart.GetFrame().Value, ToEnd.GetFrame().Value);
+		return MakeInclusiveRange(ToStart.GetFrame().Value, ToEnd.GetFrame().Value);
 	}
 }
diff --git a/Engine/Plugins/Animation/LiveLinkHub/Source/LiveLinkHub/Private/Recording/SLiveLinkHubRecordingListView.h b/Engine/Plugins/Animation/LiveLinkHub/Source/LiveLinkHub/Private/Recording/SLiveLinkHubRecordingListView.h
--- a/Engine/Plugins/Animation/LiveLinkHub/Source/LiveLinkHub/Private/Recording/SLiveLinkHubRecordingListView.h
+++ b/Engine/Plugins/Animation/LiveLinkHub/Source/LiveLinkHub/Private/Recording/SLiveLinkHubRecordingListView.h
@@ -139,6 +139,37 @@ private:
 		return bAssetsAvailableCached.GetValue() ? EVisibility::Visible : EVisibility::Hidden;
 	}
 
+	TArray<FAssetViewCustomColumn> GetCustomColumns() const
+	{
+		TArray<FAssetViewCustomColumn> ReturnValue;
+		auto ColumnStringReturningLambda = [](const FAssetData& AssetData, const FName& ColumnName) -> FString
+		{
+			int32 AssetVersion;
+			const bool bHasMetaTag = AssetData.GetTagValue(GET_MEMBER_NAME_CHECKED(ULiveLinkRecording, SavedRecordingVersion), AssetVersion);
+			if (!bHasMetaTag || AssetVersion < UE::LiveLinkHub::Private::RecordingVersions::Latest)
+			{
+				return TEXT("Outdated");
+			}
+			
+			return TEXT("Latest");
+		};
+	
+		FAssetViewCustomColumn Column;
+		Column.ColumnName = GET_MEMBER_NAME_CHECKED(ULiveLinkRecording, SavedRecordingVersion);
+		Column.DataType = UObject::FAssetRegistryTag::TT_Alphabetical;
+		Column.DisplayName = LOCTEXT("VersionName", "Version");
+		Column.OnGetColumnData = FOnGetCustomAssetColumnData::CreateLambda(ColumnStringReturningLambda);
+		Column.OnGetColumnDisplayText =
+			FOnGetCustomAssetColumnDisplayText::CreateLambda(
+				[ColumnStringReturningLambda](const FAssetData& AssetData, const FName& ColumnName)
+			{
+				return FText::FromString(ColumnStringReturningLambda(AssetData, ColumnName));
+			});
+		
+		ReturnValue.Add(Column);
+		return ReturnValue;
+	}
+	
 	/** Creates the asset picker widget for selecting a recording. */
 	TSharedRef<SWidget> CreateRecordingPicker(TOptional<FAssetData> AssetData = {})
 	{
@@ -148,6 +179,7 @@ private:
 		
 		FAssetPickerConfig AssetPickerConfig;
 		{
+			AssetPickerConfig.CustomColumns = GetCustomColumns();
 			AssetPickerConfig.SelectionMode = ESelectionMode::Single;
 			AssetPickerConfig.InitialAssetViewType = EAssetViewType::Column;
 			AssetPickerConfig.bFocusSearchBoxWhenOpened = true;
@@ -208,8 +240,10 @@ private:
 			return nullptr;
 		}
 
-		TWeakObjectPtr<UObject> SelectedAsset = SelectedAssets[0].GetAsset();
-		if (!SelectedAsset.IsValid())
+		const FAssetData& SelectedAsset = SelectedAssets[0];
+		
+		TWeakObjectPtr<UObject> SelectedAssetObject = SelectedAssets[0].GetAsset();
+		if (!SelectedAssetObject.IsValid())
 		{
 			return nullptr;
 		}
@@ -223,9 +257,9 @@ private:
 				LOCTEXT("RenameRecordingTooltip", "Rename the recording"),
 				FSlateIcon(FAppStyle::GetAppStyleSetName(), "GenericCommands.Rename"),
 				FUIAction(
-					FExecuteAction::CreateLambda([SelectedAsset, this] ()
+					FExecuteAction::CreateLambda([SelectedAssetObject, this] ()
 					{
-						if (SelectedAsset.IsValid())
+						if (SelectedAssetObject.IsValid())
 						{
 							const FContentBrowserModule& ContentBrowserModule = FModuleManager::Get().LoadModuleChecked<FContentBrowserModule>(TEXT("ContentBrowser"));
 							ContentBrowserModule.Get().ExecuteRename(AssetPicker);
@@ -240,18 +274,18 @@ private:
 				LOCTEXT("DuplicateRecordingTooltip", "Duplicate the recording"),
 				FSlateIcon(FAppStyle::GetAppStyleSetName(), "GenericCommands.Duplicate"),
 				FUIAction(
-					FExecuteAction::CreateLambda([this, SelectedAsset] ()
+					FExecuteAction::CreateLambda([this, SelectedAssetObject] ()
 					{
-						if (SelectedAsset.IsValid())
+						if (SelectedAssetObject.IsValid())
 						{
 							IAssetTools& AssetTools = FModuleManager::GetModuleChecked<FAssetToolsModule>("AssetTools").Get();
 							
 							FString TargetName;
 							FString TargetPackageName;
-							IAssetTools::Get().CreateUniqueAssetName(SelectedAsset->GetOutermost()->GetName(), TEXT("_Copy"), TargetPackageName, TargetName);
+							IAssetTools::Get().CreateUniqueAssetName(SelectedAssetObject->GetOutermost()->GetName(), TEXT("_Copy"), TargetPackageName, TargetName);
 
 							// Duplicate the asset.
-							UObject* NewAsset = AssetTools.DuplicateAsset(TargetName, FPackageName::GetLongPackagePath(TargetPackageName), SelectedAsset.Get());
+							UObject* NewAsset = AssetTools.DuplicateAsset(TargetName, FPackageName::GetLongPackagePath(TargetPackageName), SelectedAssetObject.Get());
 							FSavePackageArgs SavePackageArgs;
 							SavePackageArgs.TopLevelFlags = RF_Public | RF_Standalone;
 							SavePackageArgs.Error = GLog;
@@ -263,9 +297,9 @@ private:
 							// Unload the source recording data, as the bulk data would have been fully loaded to duplicate.
 							const FLiveLinkHubModule& LiveLinkHubModule = FModuleManager::Get().GetModuleChecked<FLiveLinkHubModule>("LiveLinkHub");
 							const TStrongObjectPtr<ULiveLinkRecording> PlaybackRecording = LiveLinkHubModule.GetPlaybackController()->GetRecording();
-							if (PlaybackRecording.Get() != SelectedAsset)
+							if (PlaybackRecording.Get() != SelectedAssetObject)
 							{
-								CastChecked<ULiveLinkUAssetRecording>(SelectedAsset)->UnloadRecordingData();
+								CastChecked<ULiveLinkUAssetRecording>(SelectedAssetObject)->UnloadRecordingData();
 							}
 
 							// There is no inherent way to update the selection of the asset picker, so instead we'll recreate one that is already selecting the new asset.
@@ -295,11 +329,11 @@ private:
 				LOCTEXT("OpenFileLocationTooltip", "Open the folder containing this file"),
 				FSlateIcon(FAppStyle::GetAppStyleSetName(), "Icons.FolderOpen"),
 				FUIAction(
-					FExecuteAction::CreateLambda([SelectedAsset] ()
+					FExecuteAction::CreateLambda([SelectedAssetObject] ()
 					{
-						if (SelectedAsset.IsValid())
+						if (SelectedAssetObject.IsValid())
 						{
-							const FString PackageName = SelectedAsset->GetPathName();
+							const FString PackageName = SelectedAssetObject->GetPathName();
 							const FString AssetFilePath = FPackageName::LongPackageNameToFilename(PackageName, FPackageName::GetAssetPackageExtension());
 							const FString AbsoluteFilePath = FPaths::ConvertRelativePathToFull(AssetFilePath);
 							const FString AssetDirectory = FPaths::GetPath(AbsoluteFilePath);
@@ -307,7 +341,29 @@ private:
 							FPlatformProcess::ExploreFolder(*AssetDirectory);
 						}
 					}),
-					FCanExecuteAction::CreateLambda([SelectedAsset] () { return SelectedAsset.IsValid(); })
+					FCanExecuteAction::CreateLambda([SelectedAssetObject] () { return SelectedAssetObject.IsValid(); })
+				)
+			);
+		}
+
+		int32 SavedRecordingVersion = 0;
+		const bool bHasRecordingValue = SelectedAsset.GetTagValue(GET_MEMBER_NAME_CHECKED(ULiveLinkRecording, SavedRecordingVersion), SavedRecordingVersion);
+		if (!bHasRecordingValue || SavedRecordingVersion < UE::LiveLinkHub::Private::RecordingVersions::Latest)
+		{
+			MenuBuilder.AddMenuEntry(
+				LOCTEXT("UpgradeRecording", "Upgrade Recording"),
+				LOCTEXT("UpgradeRecordingTooltip", "Loads the entire recording into memory, converts it to the latest version, and saves to a new file."),
+				FSlateIcon(FAppStyle::GetAppStyleSetName(), "Icons.Info"),
+				FUIAction(
+					FExecuteAction::CreateLambda([SelectedAssetObject] ()
+					{
+						if (SelectedAssetObject.IsValid())
+						{
+							const FLiveLinkHubModule& LiveLinkHubModule = FModuleManager::Get().GetModuleChecked<FLiveLinkHubModule>("LiveLinkHub");
+							LiveLinkHubModule.GetLiveLinkHub()->UpgradeAndSaveRecording(CastChecked<ULiveLinkUAssetRecording>(SelectedAssetObject.Get()));
+						}
+					}),
+					FCanExecuteAction::CreateLambda([SelectedAssetObject] () { return SelectedAssetObject.IsValid(); })
 				)
 			);
 		}
diff --git a/Engine/Plugins/Animation/LiveLinkHub/Source/LiveLinkHub/Private/UI/Widgets/SLiveLinkHubPlaybackWidget.cpp b/Engine/Plugins/Animation/LiveLinkHub/Source/LiveLinkHub/Private/UI/Widgets/SLiveLinkHubPlaybackWidget.cpp
--- a/Engine/Plugins/Animation/LiveLinkHub/Source/LiveLinkHub/Private/UI/Widgets/SLiveLinkHubPlaybackWidget.cpp
+++ b/Engine/Plugins/Animation/LiveLinkHub/Source/LiveLinkHub/Private/UI/Widgets/SLiveLinkHubPlaybackWidget.cpp
@@ -9,6 +9,7 @@
 #include "Framework/MultiBox/MultiBoxBuilder.h"
 #include "SSimpleTimeSlider.h"
 #include "Styling/SlateStyleRegistry.h"
+#include "Widgets/Images/SThrobber.h"
 #include "Widgets/Input/SButton.h"
 #include "Widgets/Input/SComboButton.h"
 #include "Widgets/Input/SSpinBox.h"
@@ -50,6 +51,9 @@ void SLiveLinkHubPlaybackWidget::Construct(const FArguments& InArgs)
 
 	IsPlaybackEnabledDelegate = InArgs._IsPlaybackEnabled;
 
+	IsUpgradingDelegate = InArgs._IsUpgrading;
+	OnGetUpgradePercentDelegate = InArgs._OnGetUpgradePercent;
+
 	OnGetFrameRate = InArgs._GetFrameRate;
 
 	OnGetFrameBufferRanges = InArgs._GetBufferRanges;
@@ -66,244 +70,257 @@ void SLiveLinkHubPlaybackWidget::Construct(const FArguments& InArgs)
 	NumberInterface->DisplayFormatChanged();
 	ChildSlot
 	[
-		SNew(SVerticalBox)
-		+SVerticalBox::Slot()
-		.AutoHeight()
-		.Padding(16.f, 8.f)
-		[
-			SNew(SLiveLinkHubTimeSlider)
-			.IsEnabled(IsPlaybackEnabled)
-			.BaseArgs(SSimpleTimeSlider::FArguments()
-			.ClampRangeHighlightSize(0.35f)
-			.ClampRangeHighlightColor(FLinearColor::Gray.CopyWithNewOpacity(0.5f))
-			.ScrubPosition_Lambda([this]()
-			{
-				// TimeSlider needs actual time in seconds, not the frame time.
-				const double Seconds = GetCurrentTime() / GetFrameRate().Numerator;
-				return Seconds;
-			})
-			.ViewRange(this, &SLiveLinkHubPlaybackWidget::GetViewRange)
-			.OnViewRangeChanged(this, &SLiveLinkHubPlaybackWidget::SetViewRange)
-			.ClampRange(this, &SLiveLinkHubPlaybackWidget::GetClampRange)
-			.OnScrubPositionChanged_Lambda(
-				[this](double NewScrubTime, bool bIsScrubbing)
-						{
-							if (bIsScrubbing)
-							{
-								//  Convert time in seconds to frame time as double.
-								SetCurrentTime(NewScrubTime * GetFrameRate().Numerator);
-							}
-						})
-			)
-			.BufferRange(this, &SLiveLinkHubPlaybackWidget::GetBufferRanges)
-		]
-		+SVerticalBox::Slot()
-		.AutoHeight()
-		.Padding(8.f, 8.f)
+		SNew(SOverlay)
+		+SOverlay::Slot()
 		[
-			SNew(SHorizontalBox)
-			+SHorizontalBox::Slot()
-			.AutoWidth()
+			SNew(SVerticalBox)
+			+SVerticalBox::Slot()
+			.AutoHeight()
+			.Padding(16.f, 8.f)
 			[
-				SNew(SButton)
+				// Box is just so we can wrap the playback control to disable it. We don't want to disable
+				// the entire vertical box because we want the eject option to remain available.
+				SNew(SBox)
 				.IsEnabled(IsPlaybackEnabled)
-				.ButtonStyle(FAppStyle::Get(), "Animation.PlayControlsButton")
-				.OnClicked(this, &SLiveLinkHubPlaybackWidget::OnFirstFramePressed)
-				.ToolTipText(LOCTEXT("ToFront", "To Front"))
-				.ContentPadding(0.0f)
 				[
-					SNew(SImage)
-					.ColorAndOpacity(FSlateColor::UseSubduedForeground())
-					.Image(FAppStyle::Get().GetBrush("Animation.Backward_End"))
+					SNew(SLiveLinkHubTimeSlider)
+					.IsEnabled(IsPlaybackEnabled)
+					.BaseArgs(SSimpleTimeSlider::FArguments()
+					.ClampRangeHighlightSize(0.35f)
+					.ClampRangeHighlightColor(FLinearColor::Gray.CopyWithNewOpacity(0.5f))
+					.ScrubPosition_Lambda([this]()
+					{
+						// TimeSlider needs actual time in seconds, not the frame time.
+						const double Seconds = GetCurrentTime() / GetFrameRate().Numerator;
+						return Seconds;
+					})
+					.ViewRange(this, &SLiveLinkHubPlaybackWidget::GetViewRange)
+					.OnViewRangeChanged(this, &SLiveLinkHubPlaybackWidget::SetViewRange)
+					.ClampRange(this, &SLiveLinkHubPlaybackWidget::GetClampRange)
+					.OnScrubPositionChanged_Lambda(
+						[this](double NewScrubTime, bool bIsScrubbing)
+								{
+									if (bIsScrubbing)
+									{
+										//  Convert time in seconds to frame time as double.
+										SetCurrentTime(NewScrubTime * GetFrameRate().Numerator);
+									}
+								})
+					)
+					.BufferRange(this, &SLiveLinkHubPlaybackWidget::GetBufferRanges)
 				]
 			]
-			+SHorizontalBox::Slot()
-			.AutoWidth()
+			+SVerticalBox::Slot()
+			.AutoHeight()
+			.Padding(8.f, 8.f)
 			[
-				SNew(SButton)
-				.IsEnabled(IsPlaybackEnabled)
-				.ButtonStyle(FAppStyle::Get(), "Animation.PlayControlsButton")
-				.OnClicked(this, &SLiveLinkHubPlaybackWidget::OnPreviousFramePressed)
-				.ToolTipText(LOCTEXT("ToPrevious", "To Previous"))
-				.ContentPadding(0.0f)
+				SNew(SHorizontalBox)
+				+SHorizontalBox::Slot()
+				.AutoWidth()
 				[
-					SNew(SImage)
-					.ColorAndOpacity(FSlateColor::UseSubduedForeground())
-					.Image(FAppStyle::Get().GetBrush("Animation.Backward_Step"))
+					SNew(SButton)
+					.IsEnabled(IsPlaybackEnabled)
+					.ButtonStyle(FAppStyle::Get(), "Animation.PlayControlsButton")
+					.OnClicked(this, &SLiveLinkHubPlaybackWidget::OnFirstFramePressed)
+					.ToolTipText(LOCTEXT("ToFront", "To Front"))
+					.ContentPadding(0.0f)
+					[
+						SNew(SImage)
+						.ColorAndOpacity(FSlateColor::UseSubduedForeground())
+						.Image(FAppStyle::Get().GetBrush("Animation.Backward_End"))
+					]
 				]
-			]
-			+SHorizontalBox::Slot()
-			.AutoWidth()
-			[
-				SNew(SButton)
-				.IsEnabled(IsPlaybackEnabled)
-				.OnClicked(this, &SLiveLinkHubPlaybackWidget::OnPlayReversePressed)
-				.ToolTipText_Lambda([this]()
-				{
-					return IsPaused() || !IsPlayingInReverse() ? LOCTEXT("ReverseButton", "Reverse") : LOCTEXT("PauseButton", "Pause");
-				})
-				.ButtonStyle( FAppStyle::Get(), "Animation.PlayControlsButton" )
-				.ContentPadding(0.0f)
+				+SHorizontalBox::Slot()
+				.AutoWidth()
 				[
-					SNew(SImage)
-					.ColorAndOpacity(FSlateColor::UseSubduedForeground())
-					.Image(this, &SLiveLinkHubPlaybackWidget::GetPlayReverseIcon)
+					SNew(SButton)
+					.IsEnabled(IsPlaybackEnabled)
+					.ButtonStyle(FAppStyle::Get(), "Animation.PlayControlsButton")
+					.OnClicked(this, &SLiveLinkHubPlaybackWidget::OnPreviousFramePressed)
+					.ToolTipText(LOCTEXT("ToPrevious", "To Previous"))
+					.ContentPadding(0.0f)
+					[
+						SNew(SImage)
+						.ColorAndOpacity(FSlateColor::UseSubduedForeground())
+						.Image(FAppStyle::Get().GetBrush("Animation.Backward_Step"))
+					]
 				]
-			]
-			+SHorizontalBox::Slot()
-			.AutoWidth()
-			[
-				SNew(SButton)
-				.IsEnabled(IsPlaybackEnabled)
-				.OnClicked(this, &SLiveLinkHubPlaybackWidget::OnPlayForwardPressed)
-				.ToolTipText_Lambda([this]()
-				{
-					return IsPaused() || IsPlayingInReverse() ? LOCTEXT("PlayButton", "Play") : LOCTEXT("PauseButton", "Pause");
-				})
-				.ButtonStyle( FAppStyle::Get(), "Animation.PlayControlsButton" )
-				.ContentPadding(0.0f)
+				+SHorizontalBox::Slot()
+				.AutoWidth()
 				[
-					SNew(SImage)
-					.ColorAndOpacity(FSlateColor::UseSubduedForeground())
-					.Image(this, &SLiveLinkHubPlaybackWidget::GetPlayForwardIcon)
+					SNew(SButton)
+					.IsEnabled(IsPlaybackEnabled)
+					.OnClicked(this, &SLiveLinkHubPlaybackWidget::OnPlayReversePressed)
+					.ToolTipText_Lambda([this]()
+					{
+						return IsPaused() || !IsPlayingInReverse() ? LOCTEXT("ReverseButton", "Reverse") : LOCTEXT("PauseButton", "Pause");
+					})
+					.ButtonStyle( FAppStyle::Get(), "Animation.PlayControlsButton" )
+					.ContentPadding(0.0f)
+					[
+						SNew(SImage)
+						.ColorAndOpacity(FSlateColor::UseSubduedForeground())
+						.Image(this, &SLiveLinkHubPlaybackWidget::GetPlayReverseIcon)
+					]
 				]
-			]
-			+SHorizontalBox::Slot()
-			.AutoWidth()
-			[
-				SNew(SButton)
-				.IsEnabled(IsPlaybackEnabled)
-				.ButtonStyle(FAppStyle::Get(), "Animation.PlayControlsButton")
-				.OnClicked(this, &SLiveLinkHubPlaybackWidget::OnNextFramePressed)
-				.ToolTipText(LOCTEXT("ToNext", "To Next"))
-				.ContentPadding(0.0f)
+				+SHorizontalBox::Slot()
+				.AutoWidth()
 				[
-					SNew(SImage)
-					.ColorAndOpacity(FSlateColor::UseSubduedForeground())
-					.Image(FAppStyle::Get().GetBrush("Animation.Forward_Step"))
+					SNew(SButton)
+					.IsEnabled(IsPlaybackEnabled)
+					.OnClicked(this, &SLiveLinkHubPlaybackWidget::OnPlayForwardPressed)
+					.ToolTipText_Lambda([this]()
+					{
+						return IsPaused() || IsPlayingInReverse() ? LOCTEXT("PlayButton", "Play") : LOCTEXT("PauseButton", "Pause");
+					})
+					.ButtonStyle( FAppStyle::Get(), "Animation.PlayControlsButton" )
+					.ContentPadding(0.0f)
+					[
+						SNew(SImage)
+						.ColorAndOpacity(FSlateColor::UseSubduedForeground())
+						.Image(this, &SLiveLinkHubPlaybackWidget::GetPlayForwardIcon)
+					]
 				]
-			]
-			+SHorizontalBox::Slot()
-			.AutoWidth()
-			[
-				SNew(SButton)
-				.IsEnabled(IsPlaybackEnabled)
-				.ButtonStyle(FAppStyle::Get(), "Animation.PlayControlsButton")
-				.OnClicked(this, &SLiveLinkHubPlaybackWidget::OnLastFramePressed)
-				.ToolTipText(LOCTEXT("ToEnd", "To End"))
-				.ContentPadding(0.0f)
+				+SHorizontalBox::Slot()
+				.AutoWidth()
 				[
-					SNew(SImage)
-					.ColorAndOpacity(FSlateColor::UseSubduedForeground())
-					.Image(FAppStyle::Get().GetBrush("Animation.Forward_End"))
+					SNew(SButton)
+					.IsEnabled(IsPlaybackEnabled)
+					.ButtonStyle(FAppStyle::Get(), "Animation.PlayControlsButton")
+					.OnClicked(this, &SLiveLinkHubPlaybackWidget::OnNextFramePressed)
+					.ToolTipText(LOCTEXT("ToNext", "To Next"))
+					.ContentPadding(0.0f)
+					[
+						SNew(SImage)
+						.ColorAndOpacity(FSlateColor::UseSubduedForeground())
+						.Image(FAppStyle::Get().GetBrush("Animation.Forward_Step"))
+					]
 				]
-			]
-			+SHorizontalBox::Slot()
-			.AutoWidth()
-			[
-				SNew(SButton)
-				.IsEnabled(IsPlaybackEnabled)
-				.ButtonStyle(FAppStyle::Get(), "Animation.PlayControlsButton")
-				.OnClicked(this, &SLiveLinkHubPlaybackWidget::OnLoopPressed)
-				.ToolTipText_Raw(this, &SLiveLinkHubPlaybackWidget::GetLoopTooltip)
-				.ContentPadding(0.0f)
+				+SHorizontalBox::Slot()
+				.AutoWidth()
 				[
-					SNew(SImage)
-					.ColorAndOpacity(FSlateColor::UseSubduedForeground())
-					.Image(this, &SLiveLinkHubPlaybackWidget::GetLoopIcon)
+					SNew(SButton)
+					.IsEnabled(IsPlaybackEnabled)
+					.ButtonStyle(FAppStyle::Get(), "Animation.PlayControlsButton")
+					.OnClicked(this, &SLiveLinkHubPlaybackWidget::OnLastFramePressed)
+					.ToolTipText(LOCTEXT("ToEnd", "To End"))
+					.ContentPadding(0.0f)
+					[
+						SNew(SImage)
+						.ColorAndOpacity(FSlateColor::UseSubduedForeground())
+						.Image(FAppStyle::Get().GetBrush("Animation.Forward_End"))
+					]
 				]
-			]
-			+ SHorizontalBox::Slot()
-			.Padding(FMargin(12.0, 0.0, 0.0, 0.0))
-			.AutoWidth()
-			[
-				SNew(SButton)
-				.ButtonStyle(FAppStyle::Get(), "Animation.PlayControlsButton")
-				.ContentPadding(FMargin(2.0f))
-				.OnClicked(this, &SLiveLinkHubPlaybackWidget::OnExitPlaybackPressed)
-				.ToolTipText(LOCTEXT("EjectButtonToolTip", "Exit Playback"))
+				+SHorizontalBox::Slot()
+				.AutoWidth()
 				[
-					SNew(SImage)
-						.Image(FAppStyle::Get().GetBrush("LevelViewport.EjectActorPilot.Small"))
+					SNew(SButton)
+					.IsEnabled(IsPlaybackEnabled)
+					.ButtonStyle(FAppStyle::Get(), "Animation.PlayControlsButton")
+					.OnClicked(this, &SLiveLinkHubPlaybackWidget::OnLoopPressed)
+					.ToolTipText_Raw(this, &SLiveLinkHubPlaybackWidget::GetLoopTooltip)
+					.ContentPadding(0.0f)
+					[
+						SNew(SImage)
 						.ColorAndOpacity(FSlateColor::UseSubduedForeground())
+						.Image(this, &SLiveLinkHubPlaybackWidget::GetLoopIcon)
+					]
+				]
+				+ SHorizontalBox::Slot()
+				.Padding(FMargin(12.0, 0.0, 0.0, 0.0))
+				.AutoWidth()
+				[
+					SNew(SButton)
+					.ButtonStyle(FAppStyle::Get(), "Animation.PlayControlsButton")
+					.ContentPadding(FMargin(2.0f))
+					.OnClicked(this, &SLiveLinkHubPlaybackWidget::OnExitPlaybackPressed)
+					.ToolTipText(LOCTEXT("EjectButtonToolTip", "Exit Playback"))
+					[
+						SNew(SImage)
+							.Image(FAppStyle::Get().GetBrush("LevelViewport.EjectActorPilot.Small"))
+							.ColorAndOpacity(FSlateColor::UseSubduedForeground())
+					]
+				]
+				+SHorizontalBox::Slot()
+				.Padding(12.f, 0.f, 0.f, 0.f)
+				[
+					SNew(SSpinBox<double>)
+					.IsEnabled(IsPlaybackEnabled)
+					.Value(this, &SLiveLinkHubPlaybackWidget::GetSelectionStartTime)
+					.ToolTipText(LOCTEXT("SelectionStartTime", "Selection start time"))
+					.OnValueChanged(this, &SLiveLinkHubPlaybackWidget::SetSelectionStartTime)
+					.MinValue(TOptional<double>())
+					.MaxValue(TOptional<double>())
+					.Style(&FAppStyle::Get().GetWidgetStyle<FSpinBoxStyle>("Sequencer.HyperlinkSpinBox"))
+					.TypeInterface(NumberInterface)
+					.ClearKeyboardFocusOnCommit(true)
+					.Delta(this, &SLiveLinkHubPlaybackWidget::GetSpinboxDelta)
+					.LinearDeltaSensitivity(25)
+				]
+				+SHorizontalBox::Slot()
+				.Padding(4.f, 0.f)
+				[
+					SNew(SSpinBox<double>)
+					.IsEnabled(IsPlaybackEnabled)
+					.Value(this, &SLiveLinkHubPlaybackWidget::GetCurrentTime)
+					.ToolTipText(LOCTEXT("CurrentTime", "Current playback time"))
+					.OnValueChanged(this, &SLiveLinkHubPlaybackWidget::SetCurrentTime)
+					.MinValue(TOptional<double>())
+					.MaxValue(TOptional<double>())
+					.Style(&FAppStyle::Get().GetWidgetStyle<FSpinBoxStyle>("Sequencer.HyperlinkSpinBox"))
+					.TypeInterface(NumberInterface)
+					.ClearKeyboardFocusOnCommit(true)
+					.Delta(this, &SLiveLinkHubPlaybackWidget::GetSpinboxDelta)
+					.LinearDeltaSensitivity(25)
+				]
+				+SHorizontalBox::Slot()
+				.Padding(4.f, 0.f)
+				[
+					SNew(SSpinBox<double>)
+					.IsEnabled(IsPlaybackEnabled)
+					.Value(this, &SLiveLinkHubPlaybackWidget::GetSelectionEndTime)
+					.ToolTipText(LOCTEXT("SelectionEndTime", "Selection end time"))
+					.OnValueChanged(this, &SLiveLinkHubPlaybackWidget::SetSelectionEndTime)
+					.MinValue(TOptional<double>())
+					.MaxValue(TOptional<double>())
+					.Style(&FAppStyle::Get().GetWidgetStyle<FSpinBoxStyle>("Sequencer.HyperlinkSpinBox"))
+					.TypeInterface(NumberInterface)
+					.ClearKeyboardFocusOnCommit(true)
+					.Delta(this, &SLiveLinkHubPlaybackWidget::GetSpinboxDelta)
+					.LinearDeltaSensitivity(25)
+				]
+				+SHorizontalBox::Slot()
+				.Padding(4.f, 0.f)
+				[
+					SNew(SSpinBox<double>)
+					.IsEnabled(IsPlaybackEnabled)
+					.Value(this, &SLiveLinkHubPlaybackWidget::GetTotalLength)
+					.IsEnabled(false)
+					.ToolTipText(LOCTEXT("RecordingLength", "Recording length"))
+					.MinValue(TOptional<double>())
+					.MaxValue(TOptional<double>())
+					.Style(&FAppStyle::Get().GetWidgetStyle<FSpinBoxStyle>("Sequencer.HyperlinkSpinBox"))
+					.TypeInterface(NumberInterface)
+				]
+				+SHorizontalBox::Slot()
+				.Padding(4.f, 0.f)
+				.AutoWidth()
+				[
+					SNew(SComboButton)
+					.IsEnabled(IsPlaybackEnabled)
+					.HasDownArrow(true)
+					.OnGetMenuContent(this, &SLiveLinkHubPlaybackWidget::MakePlaybackSettingsDropdown)
+					.ForegroundColor(FSlateColor::UseStyle())
+					.ToolTipText(LOCTEXT("PlaybackSettings_Tooltip", "Change playback settings."))
+					.ComboButtonStyle(&FAppStyle::Get().GetWidgetStyle<FComboButtonStyle>("SimpleComboButton"))
 				]
 			]
-			+SHorizontalBox::Slot()
-			.Padding(12.f, 0.f, 0.f, 0.f)
-			[
-				SNew(SSpinBox<double>)
-				.IsEnabled(IsPlaybackEnabled)
-				.Value(this, &SLiveLinkHubPlaybackWidget::GetSelectionStartTime)
-				.ToolTipText(LOCTEXT("SelectionStartTime", "Selection start time"))
-				.OnValueCommitted(this, &SLiveLinkHubPlaybackWidget::OnSelectionStartTimeCommitted)
-				.OnValueChanged(this, &SLiveLinkHubPlaybackWidget::SetSelectionStartTime)
-				.MinValue(TOptional<double>())
-				.MaxValue(TOptional<double>())
-				.Style(&FAppStyle::Get().GetWidgetStyle<FSpinBoxStyle>("Sequencer.HyperlinkSpinBox"))
-				.TypeInterface(NumberInterface)
-				.ClearKeyboardFocusOnCommit(true)
-				.Delta(this, &SLiveLinkHubPlaybackWidget::GetSpinboxDelta)
-				.LinearDeltaSensitivity(25)
-			]
-			+SHorizontalBox::Slot()
-			.Padding(4.f, 0.f)
-			[
-				SNew(SSpinBox<double>)
-				.IsEnabled(IsPlaybackEnabled)
-				.Value(this, &SLiveLinkHubPlaybackWidget::GetCurrentTime)
-				.ToolTipText(LOCTEXT("CurrentTime", "Current playback time"))
-				.OnValueCommitted(this, &SLiveLinkHubPlaybackWidget::OnCurrentTimeCommitted)
-				.OnValueChanged(this, &SLiveLinkHubPlaybackWidget::SetCurrentTime)
-				.MinValue(TOptional<double>())
-				.MaxValue(TOptional<double>())
-				.Style(&FAppStyle::Get().GetWidgetStyle<FSpinBoxStyle>("Sequencer.HyperlinkSpinBox"))
-				.TypeInterface(NumberInterface)
-				.ClearKeyboardFocusOnCommit(true)
-				.Delta(this, &SLiveLinkHubPlaybackWidget::GetSpinboxDelta)
-				.LinearDeltaSensitivity(25)
-			]
-			+SHorizontalBox::Slot()
-			.Padding(4.f, 0.f)
-			[
-				SNew(SSpinBox<double>)
-				.IsEnabled(IsPlaybackEnabled)
-				.Value(this, &SLiveLinkHubPlaybackWidget::GetSelectionEndTime)
-				.ToolTipText(LOCTEXT("SelectionEndTime", "Selection end time"))
-				.OnValueCommitted(this, &SLiveLinkHubPlaybackWidget::OnSelectionEndTimeCommitted)
-				.OnValueChanged(this, &SLiveLinkHubPlaybackWidget::SetSelectionEndTime)
-				.MinValue(TOptional<double>())
-				.MaxValue(TOptional<double>())
-				.Style(&FAppStyle::Get().GetWidgetStyle<FSpinBoxStyle>("Sequencer.HyperlinkSpinBox"))
-				.TypeInterface(NumberInterface)
-				.ClearKeyboardFocusOnCommit(true)
-				.Delta(this, &SLiveLinkHubPlaybackWidget::GetSpinboxDelta)
-				.LinearDeltaSensitivity(25)
-			]
-			+SHorizontalBox::Slot()
-			.Padding(4.f, 0.f)
-			[
-				SNew(SSpinBox<double>)
-				.IsEnabled(IsPlaybackEnabled)
-				.Value(this, &SLiveLinkHubPlaybackWidget::GetTotalLength)
-				.IsEnabled(false)
-				.ToolTipText(LOCTEXT("RecordingLength", "Recording length"))
-				.MinValue(TOptional<double>())
-				.MaxValue(TOptional<double>())
-				.Style(&FAppStyle::Get().GetWidgetStyle<FSpinBoxStyle>("Sequencer.HyperlinkSpinBox"))
-				.TypeInterface(NumberInterface)
-			]
-			+SHorizontalBox::Slot()
-			.Padding(4.f, 0.f)
-			.AutoWidth()
-			[
-				SNew(SComboButton)
-				.IsEnabled(IsPlaybackEnabled)
-				.HasDownArrow(true)
-				.OnGetMenuContent(this, &SLiveLinkHubPlaybackWidget::MakePlaybackSettingsDropdown)
-				.ForegroundColor(FSlateColor::UseStyle())
-				.ToolTipText(LOCTEXT("PlaybackSettings_Tooltip", "Change playback settings."))
-				.ComboButtonStyle(&FAppStyle::Get().GetWidgetStyle<FComboButtonStyle>("SimpleComboButton"))
-			]
+		]
+		+SOverlay::Slot()
+		.VAlign(VAlign_Fill)
+		.HAlign(HAlign_Fill)
+		[
+			MakeUpgradingWidget()
 		]
 	];
 }
@@ -366,11 +383,6 @@ void SLiveLinkHubPlaybackWidget::SetCurrentTime(double InTime)
 	OnSetCurrentTimeDelegate.Execute(FrameTime);
 }
 
-void SLiveLinkHubPlaybackWidget::OnCurrentTimeCommitted(double InTime, ETextCommit::Type InTextCommit)
-{
-	SetCurrentTime(InTime);
-}
-
 double SLiveLinkHubPlaybackWidget::GetCurrentTime() const
 {
 	check(OnGetCurrentTimeDelegate.IsBound());
@@ -393,11 +405,6 @@ double SLiveLinkHubPlaybackWidget::GetSelectionStartTime() const
 	return FrameTime.Time.AsDecimal();
 }
 
-void SLiveLinkHubPlaybackWidget::OnSelectionStartTimeCommitted(double InTime, ETextCommit::Type InTextCommit)
-{
-	SetSelectionStartTime(InTime);
-}
-
 void SLiveLinkHubPlaybackWidget::SetSelectionStartTime(double InTime)
 {
 	check(OnSetSelectionStartTimeDelegate.IsBound());
@@ -413,11 +420,6 @@ double SLiveLinkHubPlaybackWidget::GetSelectionEndTime() const
 	return FrameTime.Time.AsDecimal();
 }
 
-void SLiveLinkHubPlaybackWidget::OnSelectionEndTimeCommitted(double InTime, ETextCommit::Type InTextCommit)
-{
-	SetSelectionEndTime(InTime);
-}
-
 void SLiveLinkHubPlaybackWidget::SetSelectionEndTime(double InTime)
 {
 	check(OnSetSelectionEndTimeDelegate.IsBound());
@@ -589,6 +591,77 @@ TSharedRef<SWidget> SLiveLinkHubPlaybackWidget::MakePlaybackSettingsDropdown()
 	return MenuBuilder.MakeWidget();
 }
 
+TSharedRef<SWidget> SLiveLinkHubPlaybackWidget::MakeUpgradingWidget()
+{
+	return
+		SNew(SBorder)
+		.HAlign(HAlign_Center)
+		.VAlign(VAlign_Center)
+		.BorderImage(FAppStyle::Get().GetBrush("Brushes.Panel"))
+		.Visibility_Lambda([this]()
+		{
+			if (IsUpgradingDelegate.IsBound() && IsUpgradingDelegate.Execute())
+			{
+				return EVisibility::Visible;
+			}
+
+			return EVisibility::Collapsed;
+		})
+		[
+			SNew(SVerticalBox)
+			+ SVerticalBox::Slot()
+			.HAlign(HAlign_Center)
+			.VAlign(VAlign_Center)
+			.AutoHeight()
+			[
+				SNew(SHorizontalBox)
+				+ SHorizontalBox::Slot()
+				.Padding(6.f, 0.f)
+				.AutoWidth()
+				[
+					SNew(SCircularThrobber)
+					.Radius(6.f)
+				]
+				+ SHorizontalBox::Slot()
+				[
+					SNew(STextBlock)
+					.Text(this, &SLiveLinkHubPlaybackWidget::GetUpgradeText)
+					.Font(FAppStyle::GetFontStyle("BoldFont"))
+					.ColorAndOpacity(FSlateColor::UseForeground())
+				]
+			]
+			+ SVerticalBox::Slot()
+			.AutoHeight()
+			.HAlign(HAlign_Center)
+			.Padding(0.f, 8.f)
+			[
+				SNew(SButton)
+				.HAlign(HAlign_Center)
+				.ContentPadding(FAppStyle::GetMargin("StandardDialog.ContentPadding"))
+				.OnClicked(this, &SLiveLinkHubPlaybackWidget::OnExitPlaybackPressed)
+				.Text(LOCTEXT("CancelUpgradeText", "Cancel"))
+				.ToolTipText(LOCTEXT("UpgradeCancelToolTip", "Cancel upgrade and eject the recording."))
+			]
+		];
+}
+
+FText SLiveLinkHubPlaybackWidget::GetUpgradeText() const
+{
+	FNumberFormattingOptions NumberFormatOptions;
+	NumberFormatOptions.MaximumFractionalDigits = 0;
+	
+	const float UpgradePercent = GetUpgradePercent();
+	FText Text = FText::Format(LOCTEXT("UpgradeText", "Upgrading recording... {0}%"),
+		FText::AsNumber(UpgradePercent, &NumberFormatOptions));
+	return Text;
+}
+
+float SLiveLinkHubPlaybackWidget::GetUpgradePercent() const
+{
+	check(OnGetUpgradePercentDelegate.IsBound());
+	return OnGetUpgradePercentDelegate.Execute();
+}
+
 bool SLiveLinkHubPlaybackWidget::IsPlaybackEnabled() const
 {
 	return IsPlaybackEnabledDelegate.Execute();
diff --git a/Engine/Plugins/Animation/LiveLinkHub/Source/LiveLinkHub/Private/UI/Widgets/SLiveLinkHubPlaybackWidget.h b/Engine/Plugins/Animation/LiveLinkHub/Source/LiveLinkHub/Private/UI/Widgets/SLiveLinkHubPlaybackWidget.h
--- a/Engine/Plugins/Animation/LiveLinkHub/Source/LiveLinkHub/Private/UI/Widgets/SLiveLinkHubPlaybackWidget.h
+++ b/Engine/Plugins/Animation/LiveLinkHub/Source/LiveLinkHub/Private/UI/Widgets/SLiveLinkHubPlaybackWidget.h
@@ -26,6 +26,7 @@ public:
 	DECLARE_DELEGATE_OneParam(FOnSetDoubleRange, const TRange<double>&)
 	DECLARE_DELEGATE_RetVal(TRange<double>, FOnGetDoubleRange)
 	DECLARE_DELEGATE_RetVal(UE::LiveLinkHub::RangeHelpers::Private::TRangeArray<int32>, FOnGetIntArrayRange)
+	DECLARE_DELEGATE_RetVal(float, FOnGetFloat)
 	
 	DECLARE_DELEGATE(FOnButtonPressed);
 	
@@ -50,6 +51,10 @@ public:
 	SLATE_EVENT(FOnSetChecked, OnSetLooping)
 	/** Get if the playback widget should be enabled.*/
 	SLATE_EVENT(FOnGetChecked, IsPlaybackEnabled)
+	/** If the recording is upgrading. */
+	SLATE_EVENT(FOnGetChecked, IsUpgrading)
+	/** Get the upgrade percent. */
+	SLATE_EVENT(FOnGetFloat, OnGetUpgradePercent)
 	/** Is the recording paused? */
 	SLATE_EVENT(FOnGetChecked, IsPaused)
 	/** Is the recording playing in reverse? */
@@ -106,8 +111,6 @@ private:
 	/** Sets the current playhead time. */
 	void SetCurrentTime(double InTime);
 
-	/** When the playhead time is committed by a user change. */
-	void OnCurrentTimeCommitted(double InTime, ETextCommit::Type InTextCommit);
 	/** The current playhead time. */
 	double GetCurrentTime() const;
 	/** The total length of the recording. */
@@ -115,15 +118,11 @@ private:
 
 	/** Retrieve the selection start time. */
 	double GetSelectionStartTime() const;
-	/** When the selection start time is committed by a user change. */
-	void OnSelectionStartTimeCommitted(double InTime, ETextCommit::Type InTextCommit);
 	/** Set the selection start time. */
 	void SetSelectionStartTime(double InTime);
 
 	/** Retrieve the selection end time. */
 	double GetSelectionEndTime() const;
-	/** When the selection end time is committed by a user change. */
-	void OnSelectionEndTimeCommitted(double InTime, ETextCommit::Type InTextCommit);
 	/** Set the selection end time. */
 	void SetSelectionEndTime(double InTime);
 
@@ -176,10 +175,18 @@ private:
 	/** Create the dropdown box for playback settings. */
 	TSharedRef<SWidget> MakePlaybackSettingsDropdown();
 
+	/** The widget used for displaying upgrade notifications. */
+	TSharedRef<SWidget> MakeUpgradingWidget();
+
+	/** Get the text displayed in the upgrade notification. */
+	FText GetUpgradeText() const;
+
+	/** The percent the upgrade has finished. */
+	float GetUpgradePercent() const;
+
 	/** Returns whether the playback widgets should be enabled or not by calling the IsPlaybackEnabledDelegate delegate. */
 	bool IsPlaybackEnabled() const;
 private:
-	TWeakPtr<FLiveLinkHubPlaybackController> PlaybackController;
 	/** Delegate for pressing play forward. */
 	FOnButtonPressed OnPlayForwardDelegate;
 	
@@ -204,6 +211,12 @@ private:
 	/** Delegate checking if the recording playback is enabled */
 	FOnGetChecked IsPlaybackEnabledDelegate;
 
+	/** Delegate if recording is upgrading. */
+	FOnGetChecked IsUpgradingDelegate;
+
+	/** Delegate for upgrade percent. */
+	FOnGetFloat OnGetUpgradePercentDelegate;
+	
 	/** Delegate for checking if the recording playback is paused. */
 	FOnGetChecked OnGetPausedDelegate;
 	
