LiveLinkHub - Add prompt to verify if user intends to close the hub before closing

You can turn off this behavior by changing the "Confirm Close" option in the LiveLinkHub settings.

diff --git a/Engine/Plugins/Animation/LiveLinkHub/Source/LiveLinkHub/Private/Settings/LiveLinkHubSettings.h b/Engine/Plugins/Animation/LiveLinkHub/Source/LiveLinkHub/Private/Settings/LiveLinkHubSettings.h
--- a/Engine/Plugins/Animation/LiveLinkHub/Source/LiveLinkHub/Private/Settings/LiveLinkHubSettings.h
+++ b/Engine/Plugins/Animation/LiveLinkHub/Source/LiveLinkHub/Private/Settings/LiveLinkHubSettings.h
@@ -65,6 +65,10 @@ public:
 	UPROPERTY(config, EditAnywhere, Category = "LiveLinkHub")
 	bool bRemoveInvalidSubjectsAfterLoadingSession = false;
 
+	/** If enabled, you'll be asked to confirm if you really want to close the hub. */
+	UPROPERTY(config, EditAnywhere, Category = "LiveLinkHub")
+	bool bConfirmClose = true;
+
 	/** Whether to show the app's frame rate in the top right corner. */
 	UPROPERTY(config, EditAnywhere, Category = "LiveLinkHub")
 	bool bShowFrameRate = false;
diff --git a/Engine/Plugins/Animation/LiveLinkHub/Source/LiveLinkHub/Private/UI/Window/LiveLinkHubWindowController.cpp b/Engine/Plugins/Animation/LiveLinkHub/Source/LiveLinkHub/Private/UI/Window/LiveLinkHubWindowController.cpp
--- a/Engine/Plugins/Animation/LiveLinkHub/Source/LiveLinkHub/Private/UI/Window/LiveLinkHubWindowController.cpp
+++ b/Engine/Plugins/Animation/LiveLinkHub/Source/LiveLinkHub/Private/UI/Window/LiveLinkHubWindowController.cpp
@@ -11,6 +11,7 @@
 #include "LiveLinkHub.h"
 #include "HAL/PlatformApplicationMisc.h"
 #include "Misc/ConfigCacheIni.h"
+#include "Misc/MessageDialog.h"
 #include "UI/Widgets/SLiveLinkHubMemoryStats.h"
 #include "UI/Widgets/SLiveLinkHubTopologyModeSwitcher.h"
 #include "UI/Window/ModalWindowManager.h"
@@ -59,6 +60,8 @@ TSharedRef<SWindow> FLiveLinkHubWindowController::CreateWindow()
 	FGlobalTabmanager::Get()->SetAllowWindowMenuBar(true);
 	FSlateNotificationManager::Get().SetRootWindow(RootWindowRef);	
 
+
+	RootWindow->SetRequestDestroyWindowOverride(FRequestDestroyWindowOverride::CreateRaw(this, &FLiveLinkHubWindowController::CloseRootWindowOverride));
 	RootWindow->SetOnWindowClosed(FOnWindowClosed::CreateRaw(this, &FLiveLinkHubWindowController::OnWindowClosed));
 
 	return RootWindowRef;
@@ -138,6 +141,28 @@ TSharedPtr<FModalWindowManager> FLiveLinkHubWindowController::InitializeSlateApp
 	return MakeShared<FModalWindowManager>(CreateWindow());
 }
 
+void FLiveLinkHubWindowController::CloseRootWindowOverride(const TSharedRef<SWindow>& Window)
+{
+	if (GetDefault<ULiveLinkHubSettings>()->bConfirmClose)
+	{
+		const EAppReturnType::Type Response = FMessageDialog::Open(
+			EAppMsgCategory::Info,
+			EAppMsgType::YesNo,
+			LOCTEXT("ConfirmClose", "Are you sure you want to close Live Link Hub?"),
+			LOCTEXT("ConfirmCloseTitle", "Close Live Link Hub")
+		);
+
+		bool bOkToExit = (Response == EAppReturnType::Yes);
+		if (!bOkToExit)
+		{
+			return;
+		}
+	}
+
+	RootWindow->SetRequestDestroyWindowOverride(FRequestDestroyWindowOverride());
+	RootWindow->RequestDestroyWindow();
+}
+
 void FLiveLinkHubWindowController::OnWindowClosed(const TSharedRef<SWindow>& Window)
 {
 	if (TSharedPtr<FLiveLinkHub> LiveLinkHub = FLiveLinkHub::Get())
diff --git a/Engine/Plugins/Animation/LiveLinkHub/Source/LiveLinkHub/Private/UI/Window/LiveLinkHubWindowController.h b/Engine/Plugins/Animation/LiveLinkHub/Source/LiveLinkHub/Private/UI/Window/LiveLinkHubWindowController.h
--- a/Engine/Plugins/Animation/LiveLinkHub/Source/LiveLinkHub/Private/UI/Window/LiveLinkHubWindowController.h
+++ b/Engine/Plugins/Animation/LiveLinkHub/Source/LiveLinkHub/Private/UI/Window/LiveLinkHubWindowController.h
@@ -28,6 +28,8 @@ private:
 	TSharedRef<SWindow> CreateWindow();
 	/** Create the slate application that hosts the livelink hub. */
 	TSharedPtr<FModalWindowManager> InitializeSlateApplication();
+	/** Override to handle confirming if the user wants to quit. */
+	void CloseRootWindowOverride(const TSharedRef<SWindow>& Window);
 	/** Window closed handler. */
 	void OnWindowClosed(const TSharedRef<SWindow>& Window);
 private:
