LiveLinkHub - Implement autosave

Autosave configs are saved in .autosave files in the user appdata folder.
At the moment, only one autosave file is generated per LiveLinkHub run, but that file is periodically updated and renamed to match the save timestamp.

Also add a recent configs section to the File menu, to avoid having to manually browse to a config.

diff --git a/Engine/Plugins/Animation/LiveLinkHub/Source/LiveLinkHub/Private/LiveLinkHub.cpp b/Engine/Plugins/Animation/LiveLinkHub/Source/LiveLinkHub/Private/LiveLinkHub.cpp
--- a/Engine/Plugins/Animation/LiveLinkHub/Source/LiveLinkHub/Private/LiveLinkHub.cpp
+++ b/Engine/Plugins/Animation/LiveLinkHub/Source/LiveLinkHub/Private/LiveLinkHub.cpp
@@ -30,6 +30,7 @@
 #include "Recording/LiveLinkHubPlaybackController.h"
 #include "Recording/LiveLinkHubRecordingController.h"
 #include "Recording/LiveLinkHubRecordingListController.h"
+#include "Session/LiveLinkHubAutosaveHandler.h"
 #include "Session/LiveLinkHubSessionManager.h"
 #include "Settings/LiveLinkHubSettings.h"
 #include "Settings/LiveLinkHubTimeAndSyncSettings.h"
@@ -113,6 +114,7 @@ void FLiveLinkHub::Initialize(FLiveLinkHubTicker& Ticker)
 	SessionManager = MakeShared<FLiveLinkHubSessionManager>();
 
 	LiveLinkProvider = MakeShared<FLiveLinkHubProvider>(SessionManager.ToSharedRef(), LiveLinkHubInternalUtils::GetProviderName(InstanceId));
+	AutosaveHandler = MakeShared<FLiveLinkHubAutosaveHandler>();
 
 	FModuleManager::Get().LoadModule("Settings");
 	FModuleManager::Get().LoadModule("StatusBar");
@@ -218,6 +220,12 @@ void FLiveLinkHub::InitializePostEngineInit()
 	{ 
 		SessionManager->RestoreSession(GetDefault<ULiveLinkHubSettings>()->StartupConfig.FilePath);
 	}
+
+	// Close the playback tab when starting the app in case LLH was closed with the tab opened.
+	if (TSharedPtr<SDockTab> PlaybackTab = FLiveLinkHub::Get()->GetTabManager()->FindExistingLiveTab(UE::LiveLinkHub::PlaybackTabId))
+	{
+		PlaybackTab->RequestCloseTab();
+	}
 }
 
 void FLiveLinkHub::OnEnginePreExit()
@@ -259,10 +267,6 @@ bool FLiveLinkHub::IsRecording() const
 {
 	return RecordingController->IsRecording();
 }
-void FLiveLinkHub::Tick()
-{
-	LiveLinkHubClient->Tick();
-}
 
 TSharedRef<SWindow> FLiveLinkHub::GetRootWindow() const
 {
@@ -284,6 +288,11 @@ TSharedPtr<ILiveLinkHubSessionManager> FLiveLinkHub::GetSessionManager() const
 	return SessionManager;
 }
 
+TSharedPtr<FLiveLinkHubAutosaveHandler> FLiveLinkHub::GetAutosaveHandler() const
+{
+	return AutosaveHandler;
+}
+
 TSharedPtr<FLiveLinkHubRecordingController> FLiveLinkHub::GetRecordingController() const
 {
 	return RecordingController;
diff --git a/Engine/Plugins/Animation/LiveLinkHub/Source/LiveLinkHub/Private/LiveLinkHub.h b/Engine/Plugins/Animation/LiveLinkHub/Source/LiveLinkHub/Private/LiveLinkHub.h
--- a/Engine/Plugins/Animation/LiveLinkHub/Source/LiveLinkHub/Private/LiveLinkHub.h
+++ b/Engine/Plugins/Animation/LiveLinkHub/Source/LiveLinkHub/Private/LiveLinkHub.h
@@ -12,6 +12,7 @@
 #include "Templates/UniquePtr.h"
 #include "Toolkits/FConsoleCommandExecutor.h"
 
+class FLiveLinkHubAutosaveHandler;
 class FLiveLinkHubClient;
 class FLiveLinkHubClientsController;
 class FLiveLinkHubPlaybackController;
@@ -64,8 +65,6 @@ public:
 	/** 
 	 * Initialize LiveLinkHub. */
 	void Initialize(class FLiveLinkHubTicker& Ticker);
-	/** Tick the hub. */
-	void Tick();
 	const FLiveLinkHubInstanceId& GetId() const { return InstanceId; }
 	/** Get the root window that hosts the hub's slate application. */
 	TSharedRef<SWindow> GetRootWindow() const;
@@ -83,6 +82,8 @@ public:
 	TSharedPtr<FUICommandList> GetCommandList() const { return GetToolkitCommands(); }
 	/** Get the session manager. */
 	TSharedPtr<ILiveLinkHubSessionManager> GetSessionManager() const;
+	/** Get the autosave handler. */
+	TSharedPtr<FLiveLinkHubAutosaveHandler> GetAutosaveHandler() const;
 	/** 
 	 * Get whether the hub should act as a hub or as a spoke. 
 	 * Hubs can receive data from spokes, and transmit it to UE clients.
@@ -174,6 +175,8 @@ private:
 	TSharedPtr<FLiveLinkHubProvider> LiveLinkProvider;
 	/** Handles execution of commands. */
 	TUniquePtr<FConsoleCommandExecutor> CommandExecutor;
+	/** Handles triggering an autosave. */
+	TSharedPtr<FLiveLinkHubAutosaveHandler> AutosaveHandler;
 	/** The last opened config path. */
 	FString LastConfigPath;
 	/** LiveLinkHub instance Id, used to disambiguate different instances. */
diff --git a/Engine/Plugins/Animation/LiveLinkHub/Source/LiveLinkHub/Private/LiveLinkHubApplication.cpp b/Engine/Plugins/Animation/LiveLinkHub/Source/LiveLinkHub/Private/LiveLinkHubApplication.cpp
--- a/Engine/Plugins/Animation/LiveLinkHub/Source/LiveLinkHub/Private/LiveLinkHubApplication.cpp
+++ b/Engine/Plugins/Animation/LiveLinkHub/Source/LiveLinkHub/Private/LiveLinkHubApplication.cpp
@@ -11,6 +11,7 @@
 #include "Misc/FileHelper.h"
 #include "Misc/ScopeExit.h"
 #include "Serialization/JsonSerializer.h"
+#include "Session/LiveLinkHubAutosaveHandler.h"
 #include "StatusBarSubsystem.h"
 #include "Toolkits/AssetEditorToolkitMenuContext.h"
 #include "ToolMenus.h"
@@ -433,6 +434,49 @@ void FLiveLinkHubApplication::RegisterFileMenu(UToolMenu* Menu)
 	{
 		FileAssetSection.AddMenuEntry(FLiveLinkHubCommands::Get().NewConfig);
 		FileAssetSection.AddMenuEntry(FLiveLinkHubCommands::Get().OpenConfig);
+		FileAssetSection.AddSubMenu("RecentConfigs", LOCTEXT("RecentConfigsSubmenu", "Open Recent Configs"), LOCTEXT("RestoreConfigToolTip", "Restore a config that was auto-saved."),
+			FNewToolMenuChoice
+			(
+				FNewToolMenuDelegate::CreateLambda([](UToolMenu* ToolMenu)
+				{
+					auto CreateMenuEntry = [](const FLiveLinkHubSessionFile& SessionFile)
+					{
+						TSharedPtr<ILiveLinkHubSessionManager> SessionManager = FLiveLinkHub::Get()->GetSessionManager();
+
+						FToolMenuEntry Entry = FToolMenuEntry::InitMenuEntry(
+							*SessionFile.FileName,
+							FText::FromString(SessionFile.FileName),
+							FText::FromString(SessionFile.LastModificationDate.ToString()),
+							FSlateIcon(),
+							FExecuteAction::CreateLambda([SessionFile]() { FLiveLinkHub::Get()->GetSessionManager()->RestoreSession(SessionFile.FilePath); }),
+							EUserInterfaceActionType::None
+						);
+
+						return Entry;
+					};
+
+					// Recent config files that were saved/loaded.
+					FToolMenuSection& RecentConfigsSection = ToolMenu->AddSection("RecentConfigs", LOCTEXT("RecentConfigsLabel", "Recent Configs"));
+
+					TArray<FLiveLinkHubSessionFile> RecentFiles = GetMutableDefault<ULiveLinkHubUserSettings>()->GetRecentConfigFiles();
+
+					for (const FLiveLinkHubSessionFile& File : RecentFiles)
+					{
+						RecentConfigsSection.AddEntry(CreateMenuEntry(File));
+					}
+
+					// Config files that originated from an autosave.
+					FToolMenuSection& AutosavesSection = ToolMenu->AddSection("Autosaves", LOCTEXT("AutosaveLabel", "Autosaves"));
+
+					TArray<FLiveLinkHubSessionFile> AutosaveFiles = FLiveLinkHub::Get()->GetAutosaveHandler()->GetAutosaveFiles();
+
+					for (const FLiveLinkHubSessionFile& Autosave : AutosaveFiles)
+					{
+						AutosavesSection.AddEntry(CreateMenuEntry(Autosave));
+					}
+				})
+			)
+		);
 	}
 
 	FToolMenuSection& OpenAssetSection = FileMenu->AddSection("FileSave", LOCTEXT("FileSaveHeading", "Save"), FToolMenuInsert("FileOpen", EToolMenuInsertType::After));
diff --git a/Engine/Plugins/Animation/LiveLinkHub/Source/LiveLinkHub/Private/Session/LiveLinkHubAutosaveHandler.h b/Engine/Plugins/Animation/LiveLinkHub/Source/LiveLinkHub/Private/Session/LiveLinkHubAutosaveHandler.h
new file mode 100644
--- /dev/null
+++ b/Engine/Plugins/Animation/LiveLinkHub/Source/LiveLinkHub/Private/Session/LiveLinkHubAutosaveHandler.h
@@ -0,0 +1,96 @@
+// Copyright Epic Games, Inc. All Rights Reserved.
+
+#pragma once
+
+#include "Algo/MaxElement.h"
+#include "Algo/Transform.h"
+#include "HAL/FileManager.h"
+#include "HAL/PlatformFileManager.h"
+#include "HAL/PlatformTime.h"
+#include "LiveLinkHubSessionManager.h"
+#include "Misc/Paths.h"
+
+/** Utility class to trigger autosaves periodically. */
+class FLiveLinkHubAutosaveHandler
+{
+public:
+	inline static const FString AutosaveFileExtension = TEXT(".autosave");
+
+	FLiveLinkHubAutosaveHandler()
+	{
+		FCoreDelegates::OnBeginFrame.AddRaw(this, &FLiveLinkHubAutosaveHandler::OnBeginFrame_GameThread);
+	}
+	
+	~FLiveLinkHubAutosaveHandler()
+	{
+		FCoreDelegates::OnBeginFrame.RemoveAll(this);
+	}
+
+	/** Restore an autosave based on a file path. */
+	void RestoreAutosave(const FString& File)
+	{
+		FLiveLinkHub::Get()->GetSessionManager()->RestoreSession(File);
+	}
+
+	/** Get the list of autosave files, sorted from most recent to oldest. */
+	TArray<FLiveLinkHubSessionFile> GetAutosaveFiles() const
+	{
+		TArray<FString> Files;
+		IFileManager::Get().FindFiles(Files, *AutosaveDirectory, *AutosaveFileExtension);
+
+		TArray<FLiveLinkHubSessionFile> Autosaves;
+		Autosaves.Reserve(Files.Num());
+		Algo::Transform(Files, Autosaves, [this](const FString& InFile) { return FLiveLinkHubSessionFile{ FPaths::Combine(AutosaveDirectory, InFile) }; });
+		Algo::SortBy(Autosaves, &FLiveLinkHubSessionFile::LastModificationDate);
+		Algo::Reverse(Autosaves);
+
+		return Autosaves;
+	}
+
+private:
+	/** Attempts to save the current session if enough time has passed since the last one. */
+	void AttemptAutosave()
+	{
+		// todo: Attempt auto-save when changing session.
+		const double CurrentTime = FPlatformTime::Seconds();
+
+		if (CurrentTime - LastAutosaveTimestamp > GetDefault<ULiveLinkHubSettings>()->MinutesBetweenAutosave * 60)
+		{
+			const FString PreviousFilePath = FPaths::Combine(AutosaveDirectory, FileName);
+
+			FileName = FDateTime::Now().ToString() + AutosaveFileExtension;
+
+			UE_LOG(LogLiveLinkHub, Log, TEXT("Performing autosave to file %s"), *FileName);
+			FString SavePath = FPaths::Combine(FPlatformProcess::UserSettingsDir(), *FApp::GetEpicProductIdentifier(), TEXT("LiveLinkHub"), TEXT("Autosaves"), FileName);
+			FLiveLinkHub::Get()->GetSessionManager()->SaveCurrentSession(SavePath);
+
+			// One run of LiveLinkHub will have a single autosave file, so clear the previous file we were using since we have a new file.
+			FPlatformFileManager::Get().GetPlatformFile().DeleteFile(*PreviousFilePath);
+
+			LastAutosaveTimestamp = FPlatformTime::Seconds();
+
+			TArray<FLiveLinkHubSessionFile> Autosaves = GetAutosaveFiles();
+			if (Autosaves.Num() > GetDefault<ULiveLinkHubSettings>()->NumberOfAutosaveFilesToRetain && GetDefault<ULiveLinkHubSettings>()->NumberOfAutosaveFilesToRetain > 0)
+			{
+				FLiveLinkHubSessionFile* File = Algo::MinElementBy(Autosaves, &FLiveLinkHubSessionFile::LastModificationDate);
+
+				UE_LOG(LogLiveLinkHub, Log, TEXT("Reached max number of autosave files. Deleting autosave %s"), *File->FileName);
+				FPlatformFileManager::Get().GetPlatformFile().DeleteFile(*File->FilePath);
+			}
+		}
+	}
+
+	/** Game thread tick used to check if we need to autosave. */
+	void OnBeginFrame_GameThread()
+	{
+		AttemptAutosave();
+	}
+
+private:
+	/** Base directory for autosave files. */
+	const FString AutosaveDirectory = FPaths::Combine(FPlatformProcess::UserSettingsDir(), *FApp::GetEpicProductIdentifier(), TEXT("LiveLinkHub"), TEXT("Autosaves"));
+	/** Filename for auto-saves, based on the time the hub was opened. */
+	FString FileName = FDateTime::Now().ToString() + AutosaveFileExtension;
+	/** Timestamp of the last auto-save. */
+	double LastAutosaveTimestamp = FPlatformTime::Seconds();
+};
diff --git a/Engine/Plugins/Animation/LiveLinkHub/Source/LiveLinkHub/Private/Session/LiveLinkHubSession.h b/Engine/Plugins/Animation/LiveLinkHub/Source/LiveLinkHub/Private/Session/LiveLinkHubSession.h
--- a/Engine/Plugins/Animation/LiveLinkHub/Source/LiveLinkHub/Private/Session/LiveLinkHubSession.h
+++ b/Engine/Plugins/Animation/LiveLinkHub/Source/LiveLinkHub/Private/Session/LiveLinkHubSession.h
@@ -8,6 +8,7 @@
 #include "Engine/SystemTimeTimecodeProvider.h"
 #include "Features/IModularFeatures.h"
 #include "Framework/Application/SlateApplication.h"
+#include "HAL/FileManager.h"
 #include "LiveLinkClient.h"
 #include "LiveLinkHubModule.h"
 #include "LiveLinkHubSessionData.h"
@@ -26,6 +27,28 @@
 DECLARE_MULTICAST_DELEGATE_OneParam(FOnClientAddedToSession, FLiveLinkHubClientId);
 DECLARE_MULTICAST_DELEGATE_OneParam(FOnClientRemovedFromSession, FLiveLinkHubClientId);
 
+
+struct FLiveLinkHubSessionFile
+{
+	/** Session name. */
+	FString FileName;
+	/** Full file path to this config file. */
+	FString FilePath;
+	/** When the file was last modified. */
+	FDateTime LastModificationDate;
+	/** Whether this is an autosave. */
+	bool bWasAutosaved = false;
+
+	FLiveLinkHubSessionFile(const FString& InFile)
+		: FileName(FPaths::GetBaseFilename(InFile))
+		, FilePath(InFile)
+		, LastModificationDate(IFileManager::Get().GetTimeStamp(*InFile))
+	{
+		// At some point we could also show what subjects we had in the tooltip, or some notes about the session.
+		bWasAutosaved = FPaths::GetExtension(InFile) == TEXT("autosave");
+	}
+};
+
 /**
  * Holds the state of the hub for an active session, can be swapped out with a different session using the session manager.
  */
diff --git a/Engine/Plugins/Animation/LiveLinkHub/Source/LiveLinkHub/Private/Session/LiveLinkHubSessionManager.h b/Engine/Plugins/Animation/LiveLinkHub/Source/LiveLinkHub/Private/Session/LiveLinkHubSessionManager.h
--- a/Engine/Plugins/Animation/LiveLinkHub/Source/LiveLinkHub/Private/Session/LiveLinkHubSessionManager.h
+++ b/Engine/Plugins/Animation/LiveLinkHub/Source/LiveLinkHub/Private/Session/LiveLinkHubSessionManager.h
@@ -15,11 +15,22 @@
 #include "LiveLinkHubLog.h"
 #include "LiveLinkHubSessionExtraData.h"
 #include "LiveLinkSourceSettings.h"
+#include "Misc/ConfigCacheIni.h"
 #include "Recording/LiveLinkHubRecordingController.h"
 #include "Settings/LiveLinkHubSettings.h"
 
 #define LOCTEXT_NAMESPACE "LiveLinkHub.SessionManager"
 
+namespace UE::Private::SessionManagerUtils
+{
+	/** Returns whether the give filepath is pointing to a LiveLinkHub autosave file. */
+	inline static bool IsAutosave(const FString& FilePath)
+	{
+		return FPaths::GetExtension(FilePath) == TEXT("autosave");
+	}
+}
+
+
 DECLARE_MULTICAST_DELEGATE_OneParam(FOnActiveSessionChanged, const TSharedRef<ILiveLinkHubSession>& /**Active Session*/);
 
 class ILiveLinkHubSessionManager
@@ -49,7 +60,7 @@ public:
 	virtual void RestoreSession(FStringView InSessionPath = {}) = 0;
 
 	/** Save the current session. If not path is specified, the last save path will be used. */
-	virtual void SaveCurrentSession(const FString& SavePath = TEXT("")) = 0;
+	virtual void SaveCurrentSession(FString SavePath = TEXT("")) = 0;
 
 	/** Returns whether the current session has as already been saved to disk before. */
 	virtual bool CanSaveCurrentSession() const = 0;
@@ -134,7 +145,7 @@ public:
 		return CurrentSession;
 	}
 
-	virtual void SaveCurrentSession(const FString& SavePath = TEXT("")) override
+	virtual void SaveCurrentSession(FString SavePath = TEXT("")) override
 	{
 		if (SavePath.IsEmpty() && LastConfigPath.IsEmpty())
 		{
@@ -193,14 +204,20 @@ public:
 			Handler->OnExtraDataSessionSaving(ExtraData);
 		}
 
-		if (!SavePath.IsEmpty())
+		if (SavePath.IsEmpty())
+		{
+			// This means we're saving the current session.
+			SavePath = LastConfigPath;
+		}
+
+		if (!UE::Private::SessionManagerUtils::IsAutosave(SavePath))
 		{
+			// Don't cache the last config path for autosaves since we don't really want the user to save in that folder.
 			LastConfigPath = SavePath;
-			GetMutableDefault<ULiveLinkHubUserSettings>()->LastConfigDirectory = FPaths::GetPath(LastConfigPath);
-			GetMutableDefault<ULiveLinkHubUserSettings>()->SaveConfig();
+			GetMutableDefault<ULiveLinkHubUserSettings>()->CacheRecentConfig(SavePath);
 		}
 
-		UE::LiveLinkHub::FileUtilities::Private::SaveConfig(LiveLinkHubSessionData, LastConfigPath);
+		UE::LiveLinkHub::FileUtilities::Private::SaveConfig(LiveLinkHubSessionData, SavePath);
 	}
 
 	virtual void RestoreSession(FStringView InSessionPath /* = {} */) override
@@ -235,14 +252,23 @@ public:
 
 			if (bFileSelected && OpenFileNames.Num() > 0)
 			{
-				InSessionPath = OpenFileNames[0];
-				GetMutableDefault<ULiveLinkHubUserSettings>()->LastConfigDirectory = FPaths::GetPath(OpenFileNames[0]);
-				GetMutableDefault<ULiveLinkHubUserSettings>()->SaveConfig();
+				const FString& FilePath = OpenFileNames[0];
+				InSessionPath = FilePath;
+
+				if (!UE::Private::SessionManagerUtils::IsAutosave(FilePath))
+				{
+					GetMutableDefault<ULiveLinkHubUserSettings>()->CacheRecentConfig(FilePath);
+				}
 			}
 		}
 
 		if (!InSessionPath.IsEmpty())
 		{
+			if (!UE::Private::SessionManagerUtils::IsAutosave(FString{ InSessionPath }))
+			{
+				GetMutableDefault<ULiveLinkHubUserSettings>()->CacheRecentConfig(FString{ InSessionPath });
+			}
+
 			// Certain sources may take time to clean up. If they don't complete in time then the new config being loaded may not create
 			// duplicate sources correctly. There should be errors in the logs of the sources that failed to remove or were unable to be added.
 			constexpr bool bWaitForSourceRemoval = true;
@@ -376,11 +402,14 @@ public:
 
 private:
 	/** Load a session from disk and restore its content. */
-	void InternalRestoreSession(FStringView Path)
+	void InternalRestoreSession(FStringView Path, bool bCacheSessionPath = true)
 	{
-		LastConfigPath = Path;
+		if (bCacheSessionPath)
+		{
+			LastConfigPath = Path;
+		}
 
-		if (ULiveLinkHubSessionData* SessionData = UE::LiveLinkHub::FileUtilities::Private::LoadConfig(LastConfigPath))
+		if (ULiveLinkHubSessionData* SessionData = UE::LiveLinkHub::FileUtilities::Private::LoadConfig(FString{ Path }))
 		{
 			IModularFeatures& ModularFeatures = IModularFeatures::Get();
 			FLiveLinkHubClient* LiveLinkHubClient = static_cast<FLiveLinkHubClient*>(&ModularFeatures.GetModularFeature<ILiveLinkClient>(ILiveLinkClient::ModularFeatureName));
diff --git a/Engine/Plugins/Animation/LiveLinkHub/Source/LiveLinkHub/Private/Settings/LiveLinkHubSettings.cpp b/Engine/Plugins/Animation/LiveLinkHub/Source/LiveLinkHub/Private/Settings/LiveLinkHubSettings.cpp
--- a/Engine/Plugins/Animation/LiveLinkHub/Source/LiveLinkHub/Private/Settings/LiveLinkHubSettings.cpp
+++ b/Engine/Plugins/Animation/LiveLinkHub/Source/LiveLinkHub/Private/Settings/LiveLinkHubSettings.cpp
@@ -2,10 +2,12 @@
 
 #include "LiveLinkHubSettings.h"
 
+#include "Algo/Transform.h"
+#include "Algo/Reverse.h"
 #include "Config/LiveLinkHubTemplateTokens.h"
-#include "NamingTokensEngineSubsystem.h"
-
 #include "Engine/Engine.h"
+#include "NamingTokensEngineSubsystem.h"
+#include "Session/LiveLinkHubSession.h"
 #include "UObject/Package.h"
 
 #include UE_INLINE_GENERATED_CPP_BY_NAME(LiveLinkHubSettings)
@@ -57,3 +59,35 @@ TObjectPtr<ULiveLinkHubNamingTokens> ULiveLinkHubSettings::GetNamingTokens() con
 
 	return NamingTokens;
 }
+
+void ULiveLinkHubUserSettings::CacheRecentConfig(const FString& SavePath)
+{
+	LastConfigDirectory = FPaths::GetPath(SavePath);
+
+	// Make sure this ends up at the "top" if we touch this file.
+	RecentConfigs.Remove(SavePath);
+	RecentConfigs.Add(SavePath);
+
+	if (RecentConfigs.Num() >= MaxNumberOfRecentConfigs && MaxNumberOfRecentConfigs > 0)
+	{
+		RecentConfigs.RemoveAt(0, EAllowShrinking::No);
+	}
+
+	SaveConfig();
+}
+
+TArray<FLiveLinkHubSessionFile> ULiveLinkHubUserSettings::GetRecentConfigFiles()
+{
+	// Clear entries that no longer exist.
+	const bool bRemovedElement = RecentConfigs.RemoveAll([](const FString& FilePath) { return !FPaths::FileExists(FilePath); }) > 0;
+	if (bRemovedElement)
+	{
+		SaveConfig();
+	}
+
+	TArray<FLiveLinkHubSessionFile> Files;
+	Algo::Transform(RecentConfigs, Files, UE_PROJECTION(FLiveLinkHubSessionFile));
+	Algo::Reverse(Files);
+
+	return Files;
+}
diff --git a/Engine/Plugins/Animation/LiveLinkHub/Source/LiveLinkHub/Private/Settings/LiveLinkHubSettings.h b/Engine/Plugins/Animation/LiveLinkHub/Source/LiveLinkHub/Private/Settings/LiveLinkHubSettings.h
--- a/Engine/Plugins/Animation/LiveLinkHub/Source/LiveLinkHub/Private/Settings/LiveLinkHubSettings.h
+++ b/Engine/Plugins/Animation/LiveLinkHub/Source/LiveLinkHub/Private/Settings/LiveLinkHubSettings.h
@@ -97,6 +97,18 @@ public:
 	UPROPERTY(config, EditAnywhere, Category = "LiveLinkHub")
 	bool bPromptSaveAsOnRecord = false;
 
+	/** Whether to enable auto-save of Live Link Hub configs. */
+	UPROPERTY(config, EditAnywhere, Category = "Autosave")
+	bool bEnableAutosave = true;
+
+	/** How often should we trigger a config autosave.*/
+	UPROPERTY(config, EditAnywhere, Category = "Autosave", meta=(ClampMin="1", ClampMax="120", EditCondition="bEnableAutosave", EditConditionHides))
+	uint32 MinutesBetweenAutosave = 5;
+
+	/** How many autosave file should we retain on disk. */
+	UPROPERTY(config, EditAnywhere, Category = "Autosave", meta = (ClampMin = "1", ClampMax = "10", EditCondition = "bEnableAutosave", EditConditionHides))
+	int32 NumberOfAutosaveFilesToRetain = 10;
+
 	/** Maximum time in seconds to wait for sources to clean up. Increase this value if you notice that some sources are incorrectly cleaned up when switching a config. */
 	UPROPERTY(config, EditAnywhere, AdvancedDisplay, Category="LiveLinkHub", meta = (ClampMin="0.0"))
 	float SourceMaxCleanupTime = 0.25f;
@@ -131,6 +143,15 @@ class ULiveLinkHubUserSettings : public UObject
 	GENERATED_BODY()
 
 public:
+
+	/** Number of recent configs to keep track. */
+	UPROPERTY(config, EditAnywhere, Category = "LiveLinkHub")
+	uint8 MaxNumberOfRecentConfigs = 5;
+
+	/** File paths to recent configs that were saved/loaded. */
+	UPROPERTY(config, EditAnywhere, Category = "LiveLinkHub")
+	TArray<FString> RecentConfigs;
+
 	/** Which directories to scan to discover layouts */
 	UPROPERTY(config, EditAnywhere, Category = "LiveLinkHub")
 	TArray<FString> LayoutDirectories;
@@ -138,6 +159,13 @@ public:
 	/** The last directory of the a config that was saved or loaded. */
 	UPROPERTY(config)
 	FString LastConfigDirectory;
+
+public:
+	/** Call after an explicit config save or load to save it in the recent config file list. */
+	void CacheRecentConfig(const FString& SavePath);
+
+	/** Get the most recent config / session files that were used in LiveLinkHub. */
+	TArray<struct FLiveLinkHubSessionFile> GetRecentConfigFiles();
 };
 
 #undef UE_API
