LiveLinkHub - Change default number of autosave files to 5

diff --git a/Engine/Plugins/Animation/LiveLinkHub/Source/LiveLinkHub/Private/Settings/LiveLinkHubSettings.h b/Engine/Plugins/Animation/LiveLinkHub/Source/LiveLinkHub/Private/Settings/LiveLinkHubSettings.h
--- a/Engine/Plugins/Animation/LiveLinkHub/Source/LiveLinkHub/Private/Settings/LiveLinkHubSettings.h
+++ b/Engine/Plugins/Animation/LiveLinkHub/Source/LiveLinkHub/Private/Settings/LiveLinkHubSettings.h
@@ -107,7 +107,7 @@ public:
 
 	/** How many autosave file should we retain on disk. */
 	UPROPERTY(config, EditAnywhere, Category = "Autosave", meta = (ClampMin = "1", ClampMax = "10", EditCondition = "bEnableAutosave", EditConditionHides))
-	int32 NumberOfAutosaveFilesToRetain = 10;
+	int32 NumberOfAutosaveFilesToRetain = 5;
 
 	/** Maximum time in seconds to wait for sources to clean up. Increase this value if you notice that some sources are incorrectly cleaned up when switching a config. */
 	UPROPERTY(config, EditAnywhere, AdvancedDisplay, Category="LiveLinkHub", meta = (ClampMin="0.0"))
