LiveLinkHub: Fix asset upgrade process reporting past 100%.

diff --git a/Engine/Plugins/Animation/LiveLinkHub/Source/LiveLinkHub/Private/Recording/LiveLinkHubPlaybackController.cpp b/Engine/Plugins/Animation/LiveLinkHub/Source/LiveLinkHub/Private/Recording/LiveLinkHubPlaybackController.cpp
--- a/Engine/Plugins/Animation/LiveLinkHub/Source/LiveLinkHub/Private/Recording/LiveLinkHubPlaybackController.cpp
+++ b/Engine/Plugins/Animation/LiveLinkHub/Source/LiveLinkHub/Private/Recording/LiveLinkHubPlaybackController.cpp
@@ -111,8 +111,13 @@ TSharedRef<SWidget> FLiveLinkHubPlaybackController::MakePlaybackWidget()
 			{
 				using namespace UE::LiveLinkHub::RangeHelpers::Private;
 				const TRangeArray<int32> BufferedFrames = GetBufferedFrameRanges();
-				const float Length = GetRangeLength(BufferedFrames);
+				float Length = GetRangeLength(BufferedFrames);
 				const float MaxLength = RecordingToPlay->GetMaxFrames();
+				// Note: The range length is just an approximation. It reports the difference between the lowest frame and highest frame,
+				// which could be on different tracks. It may be possible to report a frame range slightly higher than the actual max frames are
+				// in the event the frame count of the upper end goes past the total frame length. That track's total frames should still be less
+				// than our max frames. So we clamp as a quick fix to not report over 100%.
+				Length = FMath::Clamp(Length, 0.f, MaxLength);
 				// Wait for initial load otherwise the value can jump since we do a partial load, then unload first.
 				if (MaxLength > 0.f && Recording->HasPerformedInitialLoad())
 				{
