LiveLinkHub - Pause subjects when pausing a recording so they appear as valid
Dragging playback slider will no longer pause recording on end

diff --git a/Engine/Plugins/Animation/LiveLink/Source/LiveLink/Private/LiveLinkSubject.cpp b/Engine/Plugins/Animation/LiveLink/Source/LiveLink/Private/LiveLinkSubject.cpp
--- a/Engine/Plugins/Animation/LiveLink/Source/LiveLink/Private/LiveLinkSubject.cpp
+++ b/Engine/Plugins/Animation/LiveLink/Source/LiveLink/Private/LiveLinkSubject.cpp
@@ -85,8 +85,9 @@ void FLiveLinkSubject::Update()
 		if (HasValidFrameSnapshot())
 		{
 			FrameSnapshot.FrameData.GetBaseData()->WorldTime = FPlatformTime::Seconds();
+			// Only return if we have a valid snapshot. Otherwise we'll try to recreate it, since it may have been cleared.
+			return;
 		}
-		return;
 	}
 
 	if (GetMode() == ELiveLinkSourceMode::EngineTime)
diff --git a/Engine/Plugins/Animation/LiveLinkHub/Source/LiveLinkHub/Private/Recording/LiveLinkHubPlaybackController.cpp b/Engine/Plugins/Animation/LiveLinkHub/Source/LiveLinkHub/Private/Recording/LiveLinkHubPlaybackController.cpp
--- a/Engine/Plugins/Animation/LiveLinkHub/Source/LiveLinkHub/Private/Recording/LiveLinkHubPlaybackController.cpp
+++ b/Engine/Plugins/Animation/LiveLinkHub/Source/LiveLinkHub/Private/Recording/LiveLinkHubPlaybackController.cpp
@@ -9,8 +9,9 @@
 #include "Implementations/LiveLinkUAssetRecordingPlayer.h"
 #include "LiveLinkHubCreatorAppMode.h"
 #include "LiveLinkHubLog.h"
-#include "LiveLinkRecording.h"
+#include "LiveLinkHubPlaybackSourceSettings.h"
 #include "LiveLinkPreset.h"
+#include "LiveLinkRecording.h"
 #include "LiveLinkTypes.h"
 #include "Features/IModularFeatures.h"
 #include "HAL/RunnableThread.h"
@@ -168,7 +169,9 @@ TSharedRef<SWidget> FLiveLinkHubPlaybackController::MakePlaybackWidget()
 		.IsInReverse_Raw(this, &FLiveLinkHubPlaybackController::IsPlayingInReverse)
 		.IsLooping_Raw(this, &FLiveLinkHubPlaybackController::IsLooping)
 		.OnSetLooping_Raw(this, &FLiveLinkHubPlaybackController::SetLooping)
-		.GetFrameRate_Raw(this, &FLiveLinkHubPlaybackController::GetFrameRate);
+		.GetFrameRate_Raw(this, &FLiveLinkHubPlaybackController::GetFrameRate)
+		.OnScrubbingStart_Raw(this, &FLiveLinkHubPlaybackController::OnScrubbingStart)
+		.OnScrubbingEnd_Raw(this, &FLiveLinkHubPlaybackController::OnScrubbingEnd);
 }
 
 void FLiveLinkHubPlaybackController::StartPlayback()
@@ -184,6 +187,9 @@ void FLiveLinkHubPlaybackController::ResumePlayback()
 {
 	bIsPlaying = true;
 	bIsPaused = false;
+
+	UnpausePlaybackSubjects();
+
 	FQualifiedFrameTime CurrentTime = GetCurrentTime();
 
 	// Clamp to selection start/end.
@@ -348,6 +354,7 @@ void FLiveLinkHubPlaybackController::RestartPlayback()
 void FLiveLinkHubPlaybackController::PausePlayback()
 {
 	bIsPaused = true;
+	PausePlaybackSubjects();
 }
 
 void FLiveLinkHubPlaybackController::StopPlayback()
@@ -483,6 +490,11 @@ void FLiveLinkHubPlaybackController::GoToTime(FQualifiedFrameTime InTime)
 	// Stop needs to occur to restart playback.
 	StopPlayback();
 
+	for (const FLiveLinkSubjectKey& SubjectKey : GetPlaybackSubjects())
+	{
+		Client->ClearSubjectsFrames_AnyThread(SubjectKey);
+	}
+
 	const double TimeDouble = InTime.AsSeconds();
 	
 	PlaybackStartTime -= TimeDouble;
@@ -543,6 +555,7 @@ FFrameRate FLiveLinkHubPlaybackController::GetFrameRate() const
 
 	return FFrameRate(60, 1);
 }
+
 UE::LiveLinkHub::RangeHelpers::Private::TRangeArray<int32> FLiveLinkHubPlaybackController::GetBufferedFrameRanges() const
 {
 	return RecordingPlayer->GetBufferedFrameRanges();
@@ -731,3 +744,47 @@ void FLiveLinkHubPlaybackController::OnPausePlayback()
 		PausePlayback();
 	}
 }
+
+void FLiveLinkHubPlaybackController::OnScrubbingStart()
+{
+	bStartedScrubbingPaused = IsPaused();
+
+	if (!IsPaused())
+	{
+		PausePlayback();
+	}
+}
+
+void FLiveLinkHubPlaybackController::OnScrubbingEnd()
+{
+	if (!bStartedScrubbingPaused)
+	{
+		BeginPlayback(false);
+	}
+}
+
+TArray<FLiveLinkSubjectKey> FLiveLinkHubPlaybackController::GetPlaybackSubjects() const
+{
+	TArray<FLiveLinkSubjectKey> Subjects = Client->GetSubjects(false, false);
+
+	return Subjects.FilterByPredicate([this] (const FLiveLinkSubjectKey& SubjectKey)
+	{
+		return Client->GetSourceSettings(SubjectKey.Source)->IsA<ULiveLinkHubPlaybackSourceSettings>();
+	});
+}
+
+void FLiveLinkHubPlaybackController::PausePlaybackSubjects() const
+{
+	for (const FLiveLinkSubjectKey& Subject : GetPlaybackSubjects())
+	{
+		Client->PauseSubject_AnyThread(Subject.SubjectName);
+	}
+}
+
+void FLiveLinkHubPlaybackController::UnpausePlaybackSubjects() const
+{
+	for (const FLiveLinkSubjectKey& Subject : GetPlaybackSubjects())
+	{
+		Client->UnpauseSubject_AnyThread(Subject.SubjectName);
+	}
+}
\ No newline at end of file
diff --git a/Engine/Plugins/Animation/LiveLinkHub/Source/LiveLinkHub/Private/Recording/LiveLinkHubPlaybackController.h b/Engine/Plugins/Animation/LiveLinkHub/Source/LiveLinkHub/Private/Recording/LiveLinkHubPlaybackController.h
--- a/Engine/Plugins/Animation/LiveLinkHub/Source/LiveLinkHub/Private/Recording/LiveLinkHubPlaybackController.h
+++ b/Engine/Plugins/Animation/LiveLinkHub/Source/LiveLinkHub/Private/Recording/LiveLinkHubPlaybackController.h
@@ -105,6 +105,7 @@ public:
 	/** Retrieve the current framerate. */
 	FFrameRate GetFrameRate() const;
 
+
 	/** Retrieve buffered frame range. */
 	UE::LiveLinkHub::RangeHelpers::Private::TRangeArray<int32> GetBufferedFrameRanges() const;
 	
@@ -210,6 +211,21 @@ private:
 	/** Handles the user clicking the pause button to resume playback. */
 	void OnPausePlayback();
 
+	/** Handles pausing the recording while we're scrubbing. */
+	void OnScrubbingStart();
+
+	/** Handles resuming the recording when we end scrubbing. */
+	void OnScrubbingEnd();
+
+	/** Get all playback subjects. */
+	TArray<FLiveLinkSubjectKey> GetPlaybackSubjects() const;
+
+	/** Pauses all playback subjects. */
+	void PausePlaybackSubjects() const;
+
+	/** Unpauses all playback subjects. */
+	void UnpausePlaybackSubjects() const;
+
 private:
 	/** If the system has established a connection with the client. */
 	bool bIsReady = false;
@@ -267,4 +283,7 @@ private:
 
 	/** The last sent static frame index. Used to ensure we only send static data when necessary. */
 	TMap<FLiveLinkSubjectKey, int32> LastStaticFrameIndex;
+
+	/** Keeps track of whether the recording was paused when starting to scrub to ensure we don't pause if we scrub while playing. */
+	bool bStartedScrubbingPaused = false;
 };
diff --git a/Engine/Plugins/Animation/LiveLinkHub/Source/LiveLinkHub/Private/Recording/LiveLinkHubPlaybackSourceSettings.cpp b/Engine/Plugins/Animation/LiveLinkHub/Source/LiveLinkHub/Private/Recording/LiveLinkHubPlaybackSourceSettings.cpp
new file mode 100644
--- /dev/null
+++ b/Engine/Plugins/Animation/LiveLinkHub/Source/LiveLinkHub/Private/Recording/LiveLinkHubPlaybackSourceSettings.cpp
@@ -0,0 +1,11 @@
+// Copyright Epic Games, Inc. All Rights Reserved.
+
+#include "LiveLinkHubPlaybackSourceSettings.h"
+#include "LiveLinkHubPlaybackSourceFactory.h"
+#include "LiveLinkSourceFactory.h"
+
+
+ULiveLinkHubPlaybackSourceSettings::ULiveLinkHubPlaybackSourceSettings()
+{
+	Factory = ULiveLinkHubPlaybackSourceFactory::StaticClass();
+}
diff --git a/Engine/Plugins/Animation/LiveLinkHub/Source/LiveLinkHub/Private/Recording/LiveLinkHubPlaybackSourceSettings.h b/Engine/Plugins/Animation/LiveLinkHub/Source/LiveLinkHub/Private/Recording/LiveLinkHubPlaybackSourceSettings.h
--- a/Engine/Plugins/Animation/LiveLinkHub/Source/LiveLinkHub/Private/Recording/LiveLinkHubPlaybackSourceSettings.h
+++ b/Engine/Plugins/Animation/LiveLinkHub/Source/LiveLinkHub/Private/Recording/LiveLinkHubPlaybackSourceSettings.h
@@ -4,9 +4,6 @@
 
 #include "LiveLinkSourceSettings.h"
 
-#include "LiveLinkHubPlaybackSourceFactory.h"
-#include "LiveLinkSourceFactory.h"
-
 #include "LiveLinkHubPlaybackSourceSettings.generated.h"
 
 /** Settings class for playback sources, which provides the factory class for them. */
@@ -16,8 +13,5 @@ class ULiveLinkHubPlaybackSourceSettings : public ULiveLinkSourceSettings
 public:
 	GENERATED_BODY()
 
-	ULiveLinkHubPlaybackSourceSettings()
-	{
-		Factory = ULiveLinkHubPlaybackSourceFactory::StaticClass();
-	}
+	ULiveLinkHubPlaybackSourceSettings();
 };
diff --git a/Engine/Plugins/Animation/LiveLinkHub/Source/LiveLinkHub/Private/Recording/LiveLinkPlaybackSource.h b/Engine/Plugins/Animation/LiveLinkHub/Source/LiveLinkHub/Private/Recording/LiveLinkPlaybackSource.h
--- a/Engine/Plugins/Animation/LiveLinkHub/Source/LiveLinkHub/Private/Recording/LiveLinkPlaybackSource.h
+++ b/Engine/Plugins/Animation/LiveLinkHub/Source/LiveLinkHub/Private/Recording/LiveLinkPlaybackSource.h
@@ -3,6 +3,7 @@
 #pragma once
 
 #include "ILiveLinkSource.h"
+#include "LiveLinkHubPlaybackSourceSettings.h"
 #include "LiveLinkSourceSettings.h"
 
 class ILiveLinkClient;
@@ -60,6 +61,11 @@ public:
 	{
 		return NSLOCTEXT("LiveLinkPlaybackSource", "PlaybackSourceStatus", "Playback");
 	}
+
+	virtual TSubclassOf<ULiveLinkSourceSettings> GetSettingsClass() const 
+	{
+		return ULiveLinkHubPlaybackSourceSettings::StaticClass();
+	}
 	//~ End ILiveLinkSource interface
 
 protected:
diff --git a/Engine/Plugins/Animation/LiveLinkHub/Source/LiveLinkHub/Private/UI/Widgets/SLiveLinkHubPlaybackWidget.cpp b/Engine/Plugins/Animation/LiveLinkHub/Source/LiveLinkHub/Private/UI/Widgets/SLiveLinkHubPlaybackWidget.cpp
--- a/Engine/Plugins/Animation/LiveLinkHub/Source/LiveLinkHub/Private/UI/Widgets/SLiveLinkHubPlaybackWidget.cpp
+++ b/Engine/Plugins/Animation/LiveLinkHub/Source/LiveLinkHub/Private/UI/Widgets/SLiveLinkHubPlaybackWidget.cpp
@@ -56,6 +56,9 @@ void SLiveLinkHubPlaybackWidget::Construct(const FArguments& InArgs)
 
 	OnGetFrameRate = InArgs._GetFrameRate;
 
+	OnScrubbingStartDelegate = InArgs._OnScrubbingStart;
+	OnScrubbingEndDelegate = InArgs._OnScrubbingEnd;
+
 	OnGetFrameBufferRanges = InArgs._GetBufferRanges;
 
 	const TAttribute<EFrameNumberDisplayFormats> GetDisplayFormatAttr = MakeAttributeSP(this, &SLiveLinkHubPlaybackWidget::GetDisplayFormat);
@@ -97,6 +100,8 @@ void SLiveLinkHubPlaybackWidget::Construct(const FArguments& InArgs)
 					.ViewRange(this, &SLiveLinkHubPlaybackWidget::GetViewRange)
 					.OnViewRangeChanged(this, &SLiveLinkHubPlaybackWidget::SetViewRange)
 					.ClampRange(this, &SLiveLinkHubPlaybackWidget::GetClampRange)
+					.OnBeginScrubberMovement_Lambda([this]() { OnScrubbingStartDelegate.ExecuteIfBound(); })
+					.OnEndScrubberMovement_Lambda([this]() { OnScrubbingEndDelegate.ExecuteIfBound(); })
 					.OnScrubPositionChanged_Lambda(
 						[this](double NewScrubTime, bool bIsScrubbing)
 								{
diff --git a/Engine/Plugins/Animation/LiveLinkHub/Source/LiveLinkHub/Private/UI/Widgets/SLiveLinkHubPlaybackWidget.h b/Engine/Plugins/Animation/LiveLinkHub/Source/LiveLinkHub/Private/UI/Widgets/SLiveLinkHubPlaybackWidget.h
--- a/Engine/Plugins/Animation/LiveLinkHub/Source/LiveLinkHub/Private/UI/Widgets/SLiveLinkHubPlaybackWidget.h
+++ b/Engine/Plugins/Animation/LiveLinkHub/Source/LiveLinkHub/Private/UI/Widgets/SLiveLinkHubPlaybackWidget.h
@@ -75,6 +75,11 @@ public:
 	SLATE_EVENT(FOnSetTime, SetSelectionEndTime)
 	/** Retrieve the frame rate for the recording. */
 	SLATE_EVENT(FOnGetFrame, GetFrameRate)
+
+	/** Delegate fired when scrubbing starts. */
+	SLATE_EVENT(FSimpleDelegate, OnScrubbingStart)
+	/** Delegate fired when scrubbing ends. */
+	SLATE_EVENT(FSimpleDelegate, OnScrubbingEnd)
 	
 	/** Get the view range (visible selection range). */
 	SLATE_EVENT(FOnSetDoubleRange, SetViewRange)
@@ -262,6 +267,12 @@ private:
 	/** Retrieve the frame rate. */
 	FOnGetFrame OnGetFrameRate;
 
+	/** Delegate triggered when scrubbing starts. */
+	FSimpleDelegate OnScrubbingStartDelegate;
+
+	/** Delegate triggered when scrubbing stops. */
+	FSimpleDelegate OnScrubbingEndDelegate;
+
 	/** The number interface for displaying the correct frame format. */
 	TSharedPtr<FFrameNumberInterface> NumberInterface;
 
