LiveLinkHub - Fix potential nullptr dereference

diff --git a/Engine/Plugins/Animation/LiveLinkHub/Source/LiveLinkHub/Private/Recording/LiveLinkHubPlaybackController.cpp b/Engine/Plugins/Animation/LiveLinkHub/Source/LiveLinkHub/Private/Recording/LiveLinkHubPlaybackController.cpp
--- a/Engine/Plugins/Animation/LiveLinkHub/Source/LiveLinkHub/Private/Recording/LiveLinkHubPlaybackController.cpp
+++ b/Engine/Plugins/Animation/LiveLinkHub/Source/LiveLinkHub/Private/Recording/LiveLinkHubPlaybackController.cpp
@@ -769,7 +769,8 @@ TArray<FLiveLinkSubjectKey> FLiveLinkHubPlaybackController::GetPlaybackSubjects(
 
 	return Subjects.FilterByPredicate([this] (const FLiveLinkSubjectKey& SubjectKey)
 	{
-		return Client->GetSourceSettings(SubjectKey.Source)->IsA<ULiveLinkHubPlaybackSourceSettings>();
+		const UObject* SourceSettings = Client->GetSourceSettings(SubjectKey.Source);
+		return SourceSettings && SourceSettings->IsA<ULiveLinkHubPlaybackSourceSettings>();
 	});
 }
 
