LiveLinkHub - Add a crash recovery mechanism that attempts to restore the last autosave upon crashing.

Note: When launching multiple livelinkhub instances, this system will only allow crash recovery for the first session that was launched.

diff --git a/Engine/Plugins/Animation/LiveLinkHub/Source/LiveLinkHub/Private/LiveLinkHub.cpp b/Engine/Plugins/Animation/LiveLinkHub/Source/LiveLinkHub/Private/LiveLinkHub.cpp
--- a/Engine/Plugins/Animation/LiveLinkHub/Source/LiveLinkHub/Private/LiveLinkHub.cpp
+++ b/Engine/Plugins/Animation/LiveLinkHub/Source/LiveLinkHub/Private/LiveLinkHub.cpp
@@ -17,6 +17,7 @@
 #include "LiveLinkEditorSettings.h"
 #include "LiveLinkHubClient.h"
 #include "LiveLinkHubCommands.h"
+#include "LiveLinkHubCrashRecovery.h"
 #include "LiveLinkHubCreatorAppMode.h"
 #include "LiveLinkHubMessageBusSourceSettings.h"
 #include "LiveLinkHubModule.h"
@@ -211,6 +212,8 @@ void FLiveLinkHub::InitializePostEngineInit()
 
 	GIsRunning = true;
 
+	CrashRecoveryHandler = MakePimpl<FLiveLinkHubCrashRecovery>();
+
 	FString SessionPath;
 	if (FParse::Value(FCommandLine::Get(), TEXT("SessionPath="), SessionPath) && SessionPath.Len())
 	{
diff --git a/Engine/Plugins/Animation/LiveLinkHub/Source/LiveLinkHub/Private/LiveLinkHub.h b/Engine/Plugins/Animation/LiveLinkHub/Source/LiveLinkHub/Private/LiveLinkHub.h
--- a/Engine/Plugins/Animation/LiveLinkHub/Source/LiveLinkHub/Private/LiveLinkHub.h
+++ b/Engine/Plugins/Animation/LiveLinkHub/Source/LiveLinkHub/Private/LiveLinkHub.h
@@ -177,6 +177,8 @@ private:
 	TUniquePtr<FConsoleCommandExecutor> CommandExecutor;
 	/** Handles triggering an autosave. */
 	TSharedPtr<FLiveLinkHubAutosaveHandler> AutosaveHandler;
+	/** Crash recovery handler. */
+	TPimplPtr<class FLiveLinkHubCrashRecovery> CrashRecoveryHandler;
 	/** The last opened config path. */
 	FString LastConfigPath;
 	/** LiveLinkHub instance Id, used to disambiguate different instances. */
diff --git a/Engine/Plugins/Animation/LiveLinkHub/Source/LiveLinkHub/Private/LiveLinkHubCrashRecovery.h b/Engine/Plugins/Animation/LiveLinkHub/Source/LiveLinkHub/Private/LiveLinkHubCrashRecovery.h
new file mode 100644
--- /dev/null
+++ b/Engine/Plugins/Animation/LiveLinkHub/Source/LiveLinkHub/Private/LiveLinkHubCrashRecovery.h
@@ -0,0 +1,174 @@
+// Copyright Epic Games, Inc. All Rights Reserved.
+
+#pragma once
+
+#include "Dom/JsonObject.h"
+#include "Misc/FileHelper.h"
+#include "Misc/MessageDialog.h"
+#include "Serialization/JsonSerializer.h"
+#include "Serialization/JsonWriter.h"
+#include "Session/LiveLinkHubAutosaveHandler.h"
+
+#define LOCTEXT_NAMESPACE "LiveLinkHubCrashRecovery"
+
+/**
+ * Utility class that allows recovering the last autosave file if it detects that the app did not close correctly in its last run.
+ */
+class FLiveLinkHubCrashRecovery
+{
+public:
+	FLiveLinkHubCrashRecovery()
+	{
+		// If we recover, delete the file
+		if (IFileManager::Get().FileExists(*CanaryFilePath))
+		{
+			UE_LOG(LogLiveLinkHub, Display, TEXT("Detected LiveLinkHub canary file."));
+			ReadCanaryFile();
+		}
+		else
+		{
+			// Only create a canary file for the first LLH instance.
+			CreateCanaryFile();
+		}
+	}
+
+	~FLiveLinkHubCrashRecovery()
+	{
+		if (bCreatedCanaryFile)
+		{
+			DeleteCanaryFile();
+		}
+	}
+
+private:
+	/** Read the canary file from disk and prompt the user for recovering from the last autosave. */
+	void ReadCanaryFile()
+	{
+		TArray<uint8> FileMemory;
+		FFileHelper::LoadFileToArray(FileMemory, *CanaryFilePath);
+
+		FMemoryReader Reader{ FileMemory };
+		TSharedRef<TJsonReader<UTF8CHAR>> JsonReader = TJsonReaderFactory<UTF8CHAR>::Create(&Reader);
+		TSharedPtr<FJsonObject> JsonObject = MakeShared<FJsonObject>();
+		if (FJsonSerializer::Deserialize(JsonReader, JsonObject) && JsonObject.IsValid())
+		{
+			uint32 PID = 0;
+			FString AppName;
+			if (JsonObject->TryGetNumberField(TEXT("PID"), PID) && JsonObject->TryGetStringField(TEXT("ApplicationName"), AppName))
+			{
+				// Make sure we don't accidentally detect a crash when the process is still open.
+
+				FProcHandle Handle = FPlatformProcess::OpenProcess(PID);
+				ON_SCOPE_EXIT{ FPlatformProcess::CloseProc(Handle); };
+
+				bool bIsProcRunning = FPlatformProcess::IsProcRunning(Handle);
+				bool bHandleRecycled = bIsProcRunning && AppName != FPlatformProcess::GetApplicationName(PID);
+
+				if (!bIsProcRunning || bHandleRecycled)
+				{
+					const bool bRecoverInDev = GConfig->GetBoolOrDefault(TEXT("LiveLinkHub"), TEXT("bRecoverCrashInDev"), false, GEngineIni);
+					if (!FPlatformMisc::IsDebuggerPresent() || bRecoverInDev)
+					{
+						PromptRecoverLastSave(PID);
+					}
+
+					DeleteCanaryFile();
+
+					// Make sure to recreate canary file if we've deleted it.
+					CreateCanaryFile();
+				}
+			}
+		}
+	}
+
+	/** Create a canary file that contains a PID in case another LiveLinkHub instance launches. */
+	void CreateCanaryFile()
+	{
+		const TUniquePtr<FArchive> Ar(IFileManager::Get().CreateFileWriter(*CanaryFilePath));
+		if (Ar)
+		{
+			TSharedRef<FJsonObject> JsonObject = MakeShared<FJsonObject>();
+			uint32 CurrentProcessId = FPlatformProcess::GetCurrentProcessId();
+
+			JsonObject->SetNumberField(TEXT("PID"), FPlatformProcess::GetCurrentProcessId());
+			const FString ApplicationName = FPlatformProcess::GetApplicationName(CurrentProcessId);
+
+			JsonObject->SetStringField(TEXT("ApplicationName"), ApplicationName);
+
+			const TSharedRef<TJsonWriter<UTF8CHAR>> JsonWriter = TJsonWriterFactory<UTF8CHAR>::Create(Ar.Get(), 0);
+			FJsonSerializer::Serialize(JsonObject, JsonWriter);
+
+			ensure(Ar->Close());
+
+			bCreatedCanaryFile = true;
+		}
+		else
+		{
+			UE_LOG(LogLiveLinkHub, Warning, TEXT("Could not read canary file (%s) when attempting to recover from a crash."), *CanaryFilePath);
+		}
+	}
+
+	/** Delete canary file on exit. */
+	void DeleteCanaryFile()
+	{
+		IFileManager::Get().Delete(*CanaryFilePath);
+	}
+
+	/** Prompt user to recover the last save file. */
+	void PromptRecoverLastSave(uint32 ProcessId)
+	{
+		if (TSharedPtr<SWindow> ParentWindow = FGlobalTabmanager::Get()->GetRootWindow())
+		{
+			ParentWindow->BringToFront();
+			ParentWindow->FlashWindow();
+		}
+
+		const EAppReturnType::Type Response = FMessageDialog::Open(
+			EAppMsgCategory::Info,
+			EAppMsgType::YesNo,
+			LOCTEXT("RecoverSave", "The last instance of Live Link Hub did not close correctly.\nDo you want to recover the last save file?"),
+			LOCTEXT("RecoverSaveTitle", "Live Link Hub closed unexpectedly.")
+		);
+
+		const bool bRecoverLastSave = (Response == EAppReturnType::Yes);
+		if (bRecoverLastSave)
+		{
+			RestoreLastAutosave(ProcessId);
+			UE_LOG(LogLiveLinkHub, Display, TEXT("Recovering last autosave."));
+		}
+	}
+
+	/** Restore the last autosave file that's compatible with the canary file. */
+	void RestoreLastAutosave(uint32 CompatibleProcessId)
+	{
+		TArray<FLiveLinkHubSessionFile> SaveFiles = FLiveLinkHub::Get()->GetAutosaveHandler()->GetAutosaveFiles();
+
+		for (const FLiveLinkHubSessionFile& File : SaveFiles)
+		{
+			if (TStrongObjectPtr<ULiveLinkHubSessionData> SessionDataPtr = FLiveLinkHub::Get()->GetSessionManager()->ReadSessionFile(File.FilePath))
+			{
+				// Since this object lives in the Transient package.
+				// Only try matching a save file with a process Id that matches the canary file.
+				if (SessionDataPtr->ProcessId == CompatibleProcessId)
+				{
+					UE_LOG(LogLiveLinkHub, Display, TEXT("Restoring session from save %s"), *File.FileName);
+					FLiveLinkHub::Get()->GetSessionManager()->RestoreSession(SessionDataPtr.Get());
+					return;
+				}
+			}
+		}
+
+		UE_LOG(LogLiveLinkHub, Display, TEXT("Could not find a valid autosave file to recover from crash."));
+	}
+	
+
+private:
+	/** Default canary file path. */
+	const FString CanaryFilePath = FPaths::Combine(FPlatformProcess::UserSettingsDir(), *FApp::GetEpicProductIdentifier(), TEXT("LiveLinkHub"), TEXT("LiveLinkHub.canary"));
+
+	/** Whether this instance created a canary file. (If a LiveLinkHub instance is launched while another one is running, we don't create one). */
+	bool bCreatedCanaryFile = false;
+};
+
+
+#undef LOCTEXT_NAMESPACE
diff --git a/Engine/Plugins/Animation/LiveLinkHub/Source/LiveLinkHub/Private/Session/LiveLinkHubAutosaveHandler.h b/Engine/Plugins/Animation/LiveLinkHub/Source/LiveLinkHub/Private/Session/LiveLinkHubAutosaveHandler.h
--- a/Engine/Plugins/Animation/LiveLinkHub/Source/LiveLinkHub/Private/Session/LiveLinkHubAutosaveHandler.h
+++ b/Engine/Plugins/Animation/LiveLinkHub/Source/LiveLinkHub/Private/Session/LiveLinkHubAutosaveHandler.h
@@ -18,12 +18,61 @@ public:
 
 	FLiveLinkHubAutosaveHandler()
 	{
+		FCoreDelegates::GetApplicationWillTerminateDelegate().AddRaw(this, &FLiveLinkHubAutosaveHandler::OnApplicationWillTerminate);
 		FCoreDelegates::OnBeginFrame.AddRaw(this, &FLiveLinkHubAutosaveHandler::OnBeginFrame_GameThread);
+		FLiveLinkHub::Get()->GetSessionManager()->OnActiveSessionChanged().AddRaw(this, &FLiveLinkHubAutosaveHandler::OnActiveSessionChanged);
+		FLiveLinkHub::Get()->GetSessionManager()->OnSessionSaved().AddRaw(this, &FLiveLinkHubAutosaveHandler::Autosave);
 	}
 	
 	~FLiveLinkHubAutosaveHandler()
 	{
+		if (TSharedPtr<FLiveLinkHub> LiveLinkHub = FLiveLinkHub::Get())
+		{
+			LiveLinkHub->GetSessionManager()->OnSessionSaved().RemoveAll(this);
+			LiveLinkHub->GetSessionManager()->OnActiveSessionChanged().RemoveAll(this);
+		}
+
 		FCoreDelegates::OnBeginFrame.RemoveAll(this);
+		FCoreDelegates::GetApplicationWillTerminateDelegate().RemoveAll(this);
+	}
+
+	/** Forcefully perform autosave. */
+	void Autosave()
+	{
+		if (TSharedPtr<FLiveLinkHub> Hub = FLiveLinkHub::Get())
+		{
+			if (Hub->IsInPlayback())
+			{
+				// Prevent autosave in playback to avoid overriding a session with a recording's sources/subjects.
+				return;
+			}
+
+			const FString PreviousFilePath = FPaths::Combine(AutosaveDirectory, FileName);
+
+			FileName = FDateTime::Now().ToString() + AutosaveFileExtension;
+
+			UE_LOG(LogLiveLinkHub, Log, TEXT("Performing autosave to file %s"), *FileName);
+			FString SavePath = FPaths::Combine(FPlatformProcess::UserSettingsDir(), *FApp::GetEpicProductIdentifier(), TEXT("LiveLinkHub"), TEXT("Autosaves"), FileName);
+			
+			if (TSharedPtr<ILiveLinkHubSessionManager> SessionManager = Hub->GetSessionManager())
+			{
+				SessionManager->SaveCurrentSession(SavePath);
+			}
+
+			// One run of LiveLinkHub will have a single autosave file, so clear the previous file we were using since we have a new file.
+			FPlatformFileManager::Get().GetPlatformFile().DeleteFile(*PreviousFilePath);
+
+			LastAutosaveTimestamp = FPlatformTime::Seconds();
+
+			TArray<FLiveLinkHubSessionFile> Autosaves = GetAutosaveFiles();
+			if (Autosaves.Num() > GetDefault<ULiveLinkHubSettings>()->NumberOfAutosaveFilesToRetain && GetDefault<ULiveLinkHubSettings>()->NumberOfAutosaveFilesToRetain > 0)
+			{
+				FLiveLinkHubSessionFile* File = Algo::MinElementBy(Autosaves, &FLiveLinkHubSessionFile::LastModificationDate);
+
+				UE_LOG(LogLiveLinkHub, Log, TEXT("Reached max number of autosave files. Deleting autosave %s"), *File->FileName);
+				FPlatformFileManager::Get().GetPlatformFile().DeleteFile(*File->FilePath);
+			}
+		}
 	}
 
 	/** Restore an autosave based on a file path. */
@@ -56,36 +105,28 @@ private:
 
 		if (CurrentTime - LastAutosaveTimestamp > GetDefault<ULiveLinkHubSettings>()->MinutesBetweenAutosave * 60)
 		{
-			const FString PreviousFilePath = FPaths::Combine(AutosaveDirectory, FileName);
-
-			FileName = FDateTime::Now().ToString() + AutosaveFileExtension;
-
-			UE_LOG(LogLiveLinkHub, Log, TEXT("Performing autosave to file %s"), *FileName);
-			FString SavePath = FPaths::Combine(FPlatformProcess::UserSettingsDir(), *FApp::GetEpicProductIdentifier(), TEXT("LiveLinkHub"), TEXT("Autosaves"), FileName);
-			FLiveLinkHub::Get()->GetSessionManager()->SaveCurrentSession(SavePath);
-
-			// One run of LiveLinkHub will have a single autosave file, so clear the previous file we were using since we have a new file.
-			FPlatformFileManager::Get().GetPlatformFile().DeleteFile(*PreviousFilePath);
-
-			LastAutosaveTimestamp = FPlatformTime::Seconds();
-
-			TArray<FLiveLinkHubSessionFile> Autosaves = GetAutosaveFiles();
-			if (Autosaves.Num() > GetDefault<ULiveLinkHubSettings>()->NumberOfAutosaveFilesToRetain && GetDefault<ULiveLinkHubSettings>()->NumberOfAutosaveFilesToRetain > 0)
-			{
-				FLiveLinkHubSessionFile* File = Algo::MinElementBy(Autosaves, &FLiveLinkHubSessionFile::LastModificationDate);
-
-				UE_LOG(LogLiveLinkHub, Log, TEXT("Reached max number of autosave files. Deleting autosave %s"), *File->FileName);
-				FPlatformFileManager::Get().GetPlatformFile().DeleteFile(*File->FilePath);
-			}
+			Autosave();
 		}
 	}
 
+	/** Trigger an autosave when we change session. */
+	void OnActiveSessionChanged(const TSharedRef<ILiveLinkHubSession>&)
+	{
+		Autosave();
+	}
+
 	/** Game thread tick used to check if we need to autosave. */
 	void OnBeginFrame_GameThread()
 	{
 		AttemptAutosave();
 	}
 
+	/** Try to do a last autosave when the app crashes. */
+	void OnApplicationWillTerminate()
+	{
+		Autosave();
+	}
+
 private:
 	/** Base directory for autosave files. */
 	const FString AutosaveDirectory = FPaths::Combine(FPlatformProcess::UserSettingsDir(), *FApp::GetEpicProductIdentifier(), TEXT("LiveLinkHub"), TEXT("Autosaves"));
diff --git a/Engine/Plugins/Animation/LiveLinkHub/Source/LiveLinkHub/Private/Session/LiveLinkHubSessionData.h b/Engine/Plugins/Animation/LiveLinkHub/Source/LiveLinkHub/Private/Session/LiveLinkHubSessionData.h
--- a/Engine/Plugins/Animation/LiveLinkHub/Source/LiveLinkHub/Private/Session/LiveLinkHubSessionData.h
+++ b/Engine/Plugins/Animation/LiveLinkHub/Source/LiveLinkHub/Private/Session/LiveLinkHubSessionData.h
@@ -58,6 +58,10 @@ public:
 	UPROPERTY(Instanced) // Instanced so that the contained types are serialized inline.
 	TMap<TSubclassOf<ULiveLinkHubSessionExtraData>, TObjectPtr<ULiveLinkHubSessionExtraData>> ExtraDatas;
 
+	/** ID of the process that created this session. Used by LiveLinkHub's CrashRecovery system to know which autosave files to consider. */
+	UPROPERTY()
+	uint32 ProcessId = 0;
+
 public:
 	ULiveLinkHubSessionExtraData* GetExtraData(TSubclassOf<ULiveLinkHubSessionExtraData> InExtraDataClass)
 	{
diff --git a/Engine/Plugins/Animation/LiveLinkHub/Source/LiveLinkHub/Private/Session/LiveLinkHubSessionManager.h b/Engine/Plugins/Animation/LiveLinkHub/Source/LiveLinkHub/Private/Session/LiveLinkHubSessionManager.h
--- a/Engine/Plugins/Animation/LiveLinkHub/Source/LiveLinkHub/Private/Session/LiveLinkHubSessionManager.h
+++ b/Engine/Plugins/Animation/LiveLinkHub/Source/LiveLinkHub/Private/Session/LiveLinkHubSessionManager.h
@@ -47,6 +47,9 @@ public:
 	/** Delegate called when the active session changes, which will change the list of sources, subjects and clients. */
 	virtual FOnActiveSessionChanged& OnActiveSessionChanged() = 0;
 
+	/** Delegate called when the active session is saved. (Not triggered for autosaves) */
+	virtual FSimpleMulticastDelegate& OnSessionSaved() = 0;
+
 	/** Get the current session, which holds information about which sources, subjects and clients that should be enabled in the hub at the moment. */
 	virtual TSharedPtr<ILiveLinkHubSession> GetCurrentSession() const = 0;
 
@@ -59,6 +62,12 @@ public:
 	/** Restore a session from file. If not provided, will prompt the user to browse for the file. */
 	virtual void RestoreSession(FStringView InSessionPath = {}) = 0;
 
+	/** Restore a session from an object in memory. */
+	virtual void RestoreSession(ULiveLinkHubSessionData* RestoreSession) = 0;
+
+	/** Read a session file and deserialize it to a UObject. */
+	virtual TStrongObjectPtr<ULiveLinkHubSessionData> ReadSessionFile(FStringView InSessionPath) = 0;
+
 	/** Save the current session. If not path is specified, the last save path will be used. */
 	virtual void SaveCurrentSession(FString SavePath = TEXT("")) = 0;
 
@@ -105,6 +114,12 @@ public:
 		return OnActiveSessionChangedDelegate;
 	}
 
+	virtual FSimpleMulticastDelegate& OnSessionSaved() override
+	{
+		check(IsInGameThread());
+		return OnSessionSavedDelegate;
+	}
+
 	virtual void NewSession() override
 	{
 		ClearSession();
@@ -204,20 +219,23 @@ public:
 			Handler->OnExtraDataSessionSaving(ExtraData);
 		}
 
+		LiveLinkHubSessionData->ProcessId = FPlatformProcess::GetCurrentProcessId();
+
 		if (SavePath.IsEmpty())
 		{
 			// This means we're saving the current session.
 			SavePath = LastConfigPath;
 		}
 
+		UE::LiveLinkHub::FileUtilities::Private::SaveConfig(LiveLinkHubSessionData, SavePath);
+
 		if (!UE::Private::SessionManagerUtils::IsAutosave(SavePath))
 		{
 			// Don't cache the last config path for autosaves since we don't really want the user to save in that folder.
 			LastConfigPath = SavePath;
 			GetMutableDefault<ULiveLinkHubUserSettings>()->CacheRecentConfig(SavePath);
+			OnSessionSavedDelegate.Broadcast();
 		}
-
-		UE::LiveLinkHub::FileUtilities::Private::SaveConfig(LiveLinkHubSessionData, SavePath);
 	}
 
 	virtual void RestoreSession(FStringView InSessionPath /* = {} */) override
@@ -273,8 +291,102 @@ public:
 			// duplicate sources correctly. There should be errors in the logs of the sources that failed to remove or were unable to be added.
 			constexpr bool bWaitForSourceRemoval = true;
 			ClearSession(bWaitForSourceRemoval);
-			InternalRestoreSession(InSessionPath);
+
+			if (ULiveLinkHubSessionData* SessionData = UE::LiveLinkHub::FileUtilities::Private::LoadConfig(FString{ InSessionPath }))
+			{
+				LastConfigPath = InSessionPath;
+				RestoreSession(SessionData);
+			}
+		}
+	}
+
+	virtual void RestoreSession(ULiveLinkHubSessionData* SessionData)
+	{
+		IModularFeatures& ModularFeatures = IModularFeatures::Get();
+		FLiveLinkHubClient* LiveLinkHubClient = static_cast<FLiveLinkHubClient*>(&ModularFeatures.GetModularFeature<ILiveLinkClient>(ILiveLinkClient::ModularFeatureName));
+
+		const FLiveLinkHubModule& LiveLinkHubModule = FModuleManager::Get().GetModuleChecked<FLiveLinkHubModule>("LiveLinkHub");
+		const TSharedPtr<FLiveLinkHubProvider> LiveLinkProvider = LiveLinkHubModule.GetLiveLinkProvider();
+
+		check(LiveLinkHubClient);
+		check(LiveLinkProvider);
+
+		if (SessionData)
+		{
+			for (const FLiveLinkSourcePreset& SourcePreset : SessionData->Sources)
+			{
+				LiveLinkHubClient->CreateSource(SourcePreset);
+				// Ensure stored source settings persist. CreateSource will call Source->InitializeSettings, which passes in
+				// a mutable settings object. Some sources may set "default" values on the settings object overriding the
+				// saved values from the config. We want to prevent that behavior, but we still have to call InitializeSettings, because
+				// other sources may set internal values based on the current settings' values, which is behavior we want to keep.
+				if (ULiveLinkSourceSettings* PresetSettings = SourcePreset.Settings.Get())
+				{
+					if (ULiveLinkSourceSettings* CreatedSettings = LiveLinkHubClient->GetSourceSettings(SourcePreset.Guid))
+					{
+						UEngine::FCopyPropertiesForUnrelatedObjectsParams CopyParams;
+						CopyParams.bDoDelta = false;
+						UEngine::CopyPropertiesForUnrelatedObjects(PresetSettings, CreatedSettings, CopyParams);
+					}
+				}
+			}
+
+			for (const FLiveLinkSubjectPreset& SubjectPreset : SessionData->Subjects)
+			{
+				LiveLinkHubClient->CreateSubject(SubjectPreset);
+			}
+		}
+
+		TSharedPtr<FLiveLinkHubSession> CurrentSessionPtr;
+		{
+			FScopeLock Lock(&CurrentSessionCS);
+			CurrentSessionPtr = CurrentSession = MakeShared<FLiveLinkHubSession>(SessionData, OnClientAddedToSessionDelegate, OnClientRemovedFromSessionDelegate);
+		}
+
+		for (FLiveLinkHubUEClientInfo& Client : SessionData->Clients)
+		{
+			CurrentSessionPtr->AddRestoredClient(Client);
+		}
+
+		// Restore extra data
+		TArray<ILiveLinkHubSessionExtraDataHandler*> ExtraDataHandlers =
+			ILiveLinkHubSessionExtraDataHandler::GetRegisteredHandlers();
+		for (ILiveLinkHubSessionExtraDataHandler* Handler : ExtraDataHandlers)
+		{
+			TSubclassOf<ULiveLinkHubSessionExtraData> ExtraDataClass = Handler->GetExtraDataClass();
+			ULiveLinkHubSessionExtraData* ExtraData = SessionData->GetExtraData(ExtraDataClass);
+			// Might be null.
+			Handler->OnExtraDataSessionLoaded(ExtraData);
 		}
+
+		OnActiveSessionChangedDelegate.Broadcast(CurrentSessionPtr.ToSharedRef());
+
+		LiveLinkProvider->UpdateTimecodeSettings(GetDefault<ULiveLinkHubTimeAndSyncSettings>()->TimecodeSettings);
+		LiveLinkProvider->UpdateCustomTimeStepSettings(GetDefault<ULiveLinkHubTimeAndSyncSettings>()->CustomTimeStepSettings);
+
+		if (GetDefault<ULiveLinkHubSettings>()->bRemoveInvalidSubjectsAfterLoadingSession)
+		{
+			constexpr double ActualWaitTime = 1.5;
+
+			// This is odd, but since we're blocking the game thread for a while to choose a session file, the internal time of the TimerManager hasn't been updated yet.
+			// In order for our timer to have the correct start time, we need to offset our time correctly so that it matches the internal time, else our timer would execute immediately.
+			double AdjustedWaitTime = FPlatformTime::Seconds() - OpenSessionTimestampsSeconds + ActualWaitTime;
+
+			FTimerHandle Handle;
+			GEditor->GetTimerManager()->SetTimer(Handle, FTimerDelegate::CreateRaw(this, &FLiveLinkHubSessionManager::CullUnresponsiveSubjects), AdjustedWaitTime, false);
+		}
+	}
+
+	virtual TStrongObjectPtr<ULiveLinkHubSessionData> ReadSessionFile(FStringView InSessionPath) override
+	{
+		// Return a StrongObjectPtr since this object lives in the transient package.
+		TStrongObjectPtr<ULiveLinkHubSessionData> SessionData;
+		if (ULiveLinkHubSessionData* DataPtr = UE::LiveLinkHub::FileUtilities::Private::LoadConfig(FString{ InSessionPath }))
+		{
+			SessionData = TStrongObjectPtr{ DataPtr };
+		}
+
+		return SessionData;
 	}
 
 	virtual bool CanSaveCurrentSession() const override
@@ -401,92 +513,6 @@ public:
 	//~ End ILiveLinkRecordingSessionInfo interface
 
 private:
-	/** Load a session from disk and restore its content. */
-	void InternalRestoreSession(FStringView Path, bool bCacheSessionPath = true)
-	{
-		if (bCacheSessionPath)
-		{
-			LastConfigPath = Path;
-		}
-
-		if (ULiveLinkHubSessionData* SessionData = UE::LiveLinkHub::FileUtilities::Private::LoadConfig(FString{ Path }))
-		{
-			IModularFeatures& ModularFeatures = IModularFeatures::Get();
-			FLiveLinkHubClient* LiveLinkHubClient = static_cast<FLiveLinkHubClient*>(&ModularFeatures.GetModularFeature<ILiveLinkClient>(ILiveLinkClient::ModularFeatureName));
-
-			const FLiveLinkHubModule& LiveLinkHubModule = FModuleManager::Get().GetModuleChecked<FLiveLinkHubModule>("LiveLinkHub");
-			const TSharedPtr<FLiveLinkHubProvider> LiveLinkProvider = LiveLinkHubModule.GetLiveLinkProvider();
-
-			check(LiveLinkHubClient);
-			check(LiveLinkProvider);
-
-			if (SessionData)
-			{
-				for (const FLiveLinkSourcePreset& SourcePreset : SessionData->Sources)
-				{
-					LiveLinkHubClient->CreateSource(SourcePreset);
-					// Ensure stored source settings persist. CreateSource will call Source->InitializeSettings, which passes in
-					// a mutable settings object. Some sources may set "default" values on the settings object overriding the
-					// saved values from the config. We want to prevent that behavior, but we still have to call InitializeSettings, because
-					// other sources may set internal values based on the current settings' values, which is behavior we want to keep.
-					if (ULiveLinkSourceSettings* PresetSettings = SourcePreset.Settings.Get())
-					{
-						if (ULiveLinkSourceSettings* CreatedSettings = LiveLinkHubClient->GetSourceSettings(SourcePreset.Guid))
-						{
-							UEngine::FCopyPropertiesForUnrelatedObjectsParams CopyParams;
-							CopyParams.bDoDelta = false;
-							UEngine::CopyPropertiesForUnrelatedObjects(PresetSettings, CreatedSettings, CopyParams);
-						}
-					}
-				}
-
-				for (const FLiveLinkSubjectPreset& SubjectPreset : SessionData->Subjects)
-				{
-					LiveLinkHubClient->CreateSubject(SubjectPreset);
-				}
-			}
-
-			TSharedPtr<FLiveLinkHubSession> CurrentSessionPtr;
-			{
-				FScopeLock Lock(&CurrentSessionCS);
-				CurrentSessionPtr = CurrentSession = MakeShared<FLiveLinkHubSession>(SessionData, OnClientAddedToSessionDelegate, OnClientRemovedFromSessionDelegate);
-			}
-
-			for (FLiveLinkHubUEClientInfo& Client : SessionData->Clients)
-			{
-				CurrentSessionPtr->AddRestoredClient(Client);
-			}
-
-			// Restore extra data
-			TArray<ILiveLinkHubSessionExtraDataHandler*> ExtraDataHandlers =
-				ILiveLinkHubSessionExtraDataHandler::GetRegisteredHandlers();
-			for (ILiveLinkHubSessionExtraDataHandler* Handler : ExtraDataHandlers)
-			{
-				TSubclassOf<ULiveLinkHubSessionExtraData> ExtraDataClass = Handler->GetExtraDataClass();
-				ULiveLinkHubSessionExtraData* ExtraData = SessionData->GetExtraData(ExtraDataClass);
-				// Might be null.
-				Handler->OnExtraDataSessionLoaded(ExtraData);
-			}
-
-			OnActiveSessionChangedDelegate.Broadcast(CurrentSessionPtr.ToSharedRef());
-
-			LiveLinkProvider->UpdateTimecodeSettings(GetDefault<ULiveLinkHubTimeAndSyncSettings>()->TimecodeSettings);
-			LiveLinkProvider->UpdateCustomTimeStepSettings(GetDefault<ULiveLinkHubTimeAndSyncSettings>()->CustomTimeStepSettings);
-
-			if (GetDefault<ULiveLinkHubSettings>()->bRemoveInvalidSubjectsAfterLoadingSession)
-			{
-				constexpr double ActualWaitTime = 1.5;
-
-				// This is odd, but since we're blocking the game thread for a while to choose a session file, the internal time of the TimerManager hasn't been updated yet.
-				// In order for our timer to have the correct start time, we need to offset our time correctly so that it matches the internal time, else our timer would execute immediately.
-				double AdjustedWaitTime = FPlatformTime::Seconds() - OpenSessionTimestampsSeconds + ActualWaitTime;
-
-				FTimerHandle Handle;
-				GEditor->GetTimerManager()->SetTimer(Handle, FTimerDelegate::CreateRaw(this, &FLiveLinkHubSessionManager::CullUnresponsiveSubjects), AdjustedWaitTime, false);
-			}
-		}
-	}
-
 	/** Clear the hub data contained in the current session, resetting the hub to its default state. */
 	void ClearSession(bool bWaitForSourceRemoval = false)
 	{
@@ -555,6 +581,9 @@ private:
 	/** Delegate triggered when the current session is changed. */
 	FOnActiveSessionChanged OnActiveSessionChangedDelegate;
 
+	/** Delegate called when the session is saved. */
+	FSimpleMulticastDelegate OnSessionSavedDelegate;
+
 	/** Delegate triggered when the recording session name is changed. */
 	FOnSessionStringChanged OnSessionNameChangedDelegate;
 
