LiveLinkHub - Add playback mode, which is a compact view of the playback widget.

Keyboard Shortcuts for Recording Playback
- Left / Backspace: Step back one frame (hold to scrub backward)
- Right: Step forward one frame (hold to scrub forward)
- Space / Enter: Toggle pause

Added a "Select Recording" button when no recording is selected, which opens a pop up window with the recording list.

diff --git a/Engine/Plugins/Animation/LiveLinkHub/Source/LiveLinkHub/Private/LiveLinkHub.cpp b/Engine/Plugins/Animation/LiveLinkHub/Source/LiveLinkHub/Private/LiveLinkHub.cpp
--- a/Engine/Plugins/Animation/LiveLinkHub/Source/LiveLinkHub/Private/LiveLinkHub.cpp
+++ b/Engine/Plugins/Animation/LiveLinkHub/Source/LiveLinkHub/Private/LiveLinkHub.cpp
@@ -19,6 +19,7 @@
 #include "LiveLinkHubCommands.h"
 #include "LiveLinkHubCrashRecovery.h"
 #include "LiveLinkHubCreatorAppMode.h"
+#include "LiveLinkHubPlaybackAppMode.h"
 #include "LiveLinkHubMessageBusSourceSettings.h"
 #include "LiveLinkHubModule.h"
 #include "LiveLinkHubSubjectSettings.h"
@@ -176,6 +177,7 @@ void FLiveLinkHub::InitializePostEngineInit()
 	// Note: Registering modes *must* happen after the layout was restored, since RestoreLayout will initiate the asset editor. 
 	// Otherwise, the mode tabs will get registered with the global tab manager instead of the host app's tab manager.
 	AddLiveLinkHubApplicationMode("CreatorMode", MakeShared<FLiveLinkHubCreatorAppMode>(SharedThis<FLiveLinkHubApplicationBase>(this)));
+	AddLiveLinkHubApplicationMode("PlaybackMode", MakeShared<FLiveLinkHubPlaybackAppMode>(SharedThis<FLiveLinkHubApplicationBase>(this)));
 
 	for (ILiveLinkHubApplicationModeFactory* Factory : IModularFeatures::Get().GetModularFeatureImplementations<ILiveLinkHubApplicationModeFactory>(ILiveLinkHubApplicationModeFactory::ModularFeatureName))
 	{
diff --git a/Engine/Plugins/Animation/LiveLinkHub/Source/LiveLinkHub/Private/LiveLinkHubCreatorAppMode.h b/Engine/Plugins/Animation/LiveLinkHub/Source/LiveLinkHub/Private/LiveLinkHubCreatorAppMode.h
--- a/Engine/Plugins/Animation/LiveLinkHub/Source/LiveLinkHub/Private/LiveLinkHubCreatorAppMode.h
+++ b/Engine/Plugins/Animation/LiveLinkHub/Source/LiveLinkHub/Private/LiveLinkHubCreatorAppMode.h
@@ -228,7 +228,7 @@ public:
 
 	virtual TSharedRef<SWidget> CreateTabBody(const FWorkflowTabSpawnInfo& Info) const override
 	{
-		return FLiveLinkHub::Get()->GetPlaybackController()->MakePlaybackWidget();
+		return FLiveLinkHub::Get()->GetPlaybackController()->GetPlaybackWidget();
 	}
 
 	virtual FText GetTabToolTipText(const FWorkflowTabSpawnInfo& Info) const override
diff --git a/Engine/Plugins/Animation/LiveLinkHub/Source/LiveLinkHub/Private/LiveLinkHubPlaybackAppMode.h b/Engine/Plugins/Animation/LiveLinkHub/Source/LiveLinkHub/Private/LiveLinkHubPlaybackAppMode.h
new file mode 100644
--- /dev/null
+++ b/Engine/Plugins/Animation/LiveLinkHub/Source/LiveLinkHub/Private/LiveLinkHubPlaybackAppMode.h
@@ -0,0 +1,76 @@
+// Copyright Epic Games, Inc. All Rights Reserved.
+
+#pragma once
+
+#include "LiveLinkHubApplicationMode.h"
+
+#include "Framework/Docking/TabManager.h"
+#include "LiveLinkHubApplicationBase.h"
+#include "LiveLinkHub.h"
+#include "Recording/LiveLinkHubPlaybackController.h"
+#include "WorkflowOrientedApp/ApplicationMode.h"
+#include "WorkflowOrientedApp/WorkflowTabFactory.h"
+#include "WorkflowOrientedApp/WorkflowTabManager.h"
+
+
+#define LOCTEXT_NAMESPACE "LiveLinkHubPlaybackMode"
+
+/** Compact mode focused exclusively on playing back recordings. */
+class FLiveLinkHubPlaybackAppMode : public FLiveLinkHubApplicationMode
+{
+public:
+	FLiveLinkHubPlaybackAppMode(TSharedPtr<FLiveLinkHubApplicationBase> App)
+		: FLiveLinkHubApplicationMode("PlaybackMode", LOCTEXT("PlaybackModeLabel", "Playback"), App)
+	{
+		TabLayout = FTabManager::NewLayout("LiveLinkHubPlaybackMode_v1.0");
+		TabLayout->AddArea
+		(
+			FTabManager::NewPrimaryArea()
+			->SetOrientation(Orient_Vertical)
+			->Split
+			(
+				FTabManager::NewStack()
+				->SetHideTabWell(true)
+				->SetSizeCoefficient(1.f)
+				->AddTab(UE::LiveLinkHub::PlaybackTabId, ETabState::OpenedTab)
+			)
+		);
+	}
+
+	//~ Begin FLiveLinkHubApplicationMode interface
+	virtual void PostActivateMode() override
+	{
+		if (TSharedPtr<FLiveLinkHubPlaybackController> PlaybackController = FLiveLinkHub::Get()->GetPlaybackController())
+		{
+			PlaybackController->SetClosePlaybackTabOnEject(false);
+			
+			PreviousWindowSize = FLiveLinkHub::Get()->GetRootWindow()->GetSizeInScreen();
+
+			FLiveLinkHub::Get()->GetRootWindow()->Resize({ 800.0, 220.0 });
+		}
+	}
+
+	virtual void PreDeactivateMode() override
+	{
+		if (TSharedPtr<FLiveLinkHubPlaybackController> PlaybackController = FLiveLinkHub::Get()->GetPlaybackController())
+		{
+			PlaybackController->SetClosePlaybackTabOnEject(true);
+
+			if (PreviousWindowSize.Size() != 0)
+			{
+				FLiveLinkHub::Get()->GetRootWindow()->Resize({ PreviousWindowSize.X, PreviousWindowSize.Y });
+			}
+		}
+	}
+
+	virtual FSlateIcon GetModeIcon() const
+	{
+		return FSlateIcon(UE::LiveLinkHub::LiveLinkStyleName, TEXT("LiveLinkHub.Playback.Icon"));
+	}
+	//~ End FLiveLinkHubApplicationMode interface
+
+private:
+	FVector2d PreviousWindowSize;
+};
+
+#undef LOCTEXT_NAMESPACE
diff --git a/Engine/Plugins/Animation/LiveLinkHub/Source/LiveLinkHub/Private/Recording/LiveLinkHubPlaybackController.cpp b/Engine/Plugins/Animation/LiveLinkHub/Source/LiveLinkHub/Private/Recording/LiveLinkHubPlaybackController.cpp
--- a/Engine/Plugins/Animation/LiveLinkHub/Source/LiveLinkHub/Private/Recording/LiveLinkHubPlaybackController.cpp
+++ b/Engine/Plugins/Animation/LiveLinkHub/Source/LiveLinkHub/Private/Recording/LiveLinkHubPlaybackController.cpp
@@ -25,6 +25,8 @@
 #include "UObject/Object.h"
 #include "UObject/Package.h"
 
+#define LOCTEXT_NAMESPACE "LiveLinkHubPlayback"
+
 /**
  * Thread safe way to load and store a FQualifiedFrameTime. This is necessary because atomic<FQualifiedFrameTime> isn't
  * necessarily cross platform compatible when performing copy operations.
@@ -93,10 +95,15 @@ FLiveLinkHubPlaybackController::~FLiveLinkHubPlaybackController()
 	}
 }
 
-TSharedRef<SWidget> FLiveLinkHubPlaybackController::MakePlaybackWidget()
+TSharedRef<SWidget> FLiveLinkHubPlaybackController::GetPlaybackWidget()
 {
-	return SNew(SLiveLinkHubPlaybackWidget)
-		.Visibility_Lambda([this]() { return IsInPlayback() ? EVisibility::Visible : EVisibility::Collapsed; })
+	if (TSharedPtr<SWidget> WidgetPtr = PlaybackWidget.Pin())
+	{
+		return WidgetPtr.ToSharedRef();
+	}
+
+	return SAssignNew(PlaybackWidget, SLiveLinkHubPlaybackWidget)
+		.Visibility_Lambda([this]() { return IsInPlayback() || !bClosePlaybackTabOnEject ? EVisibility::Visible : EVisibility::Collapsed; })
 		.IsPlaybackEnabled_Raw(this, &FLiveLinkHubPlaybackController::IsReady)
 		.IsUpgrading_Lambda([this]()
 		{
@@ -139,16 +146,35 @@ TSharedRef<SWidget> FLiveLinkHubPlaybackController::MakePlaybackWidget()
 		})
 		.OnPreviousFrame_Lambda([this]()
 		{
-			GoToTime(FQualifiedFrameTime(FFrameTime(GetCurrentFrame() - 1), GetFrameRate()));
+			FQualifiedFrameTime CurrentTimeCopy = GetCurrentTime();
+			CurrentTimeCopy.Time -= 1;
+
+			// Clamp to selection start/end.
+			if (CurrentTimeCopy.AsSeconds() < GetSelectionStartTime().AsSeconds())
+			{
+				CurrentTimeCopy = IsLooping() ? GetSelectionEndTime() : GetSelectionStartTime();
+			}
+
+			GoToTime(CurrentTimeCopy);
 		})
 		.OnNextFrame_Lambda([this]()
 		{
-			GoToTime(FQualifiedFrameTime(FFrameTime(GetCurrentFrame() + 1), GetFrameRate()));
+			FQualifiedFrameTime CurrentTimeCopy = GetCurrentTime();
+			CurrentTimeCopy.Time += 1;
+
+			// Clamp to selection start/end.
+			if (CurrentTimeCopy.AsSeconds() >= GetSelectionEndTime().AsSeconds())
+			{
+				CurrentTimeCopy = IsLooping() ? GetSelectionStartTime() : GetSelectionEndTime();
+			}
+
+			GoToTime(CurrentTimeCopy);
 		})
 		.OnExitPlayback_Lambda([this]()
 		{
-			EjectAndUnload();
+			EjectAndUnload([]() {}, nullptr, bClosePlaybackTabOnEject);
 		})
+		.OnSelectRecording_Raw(this, &FLiveLinkHubPlaybackController::OnSelectRecording)
 		.SetCurrentTime_Raw(this, &FLiveLinkHubPlaybackController::GoToTime)
 		.GetViewRange_Lambda([this]()
 		{
@@ -211,6 +237,18 @@ void FLiveLinkHubPlaybackController::ResumePlayback()
 
 void FLiveLinkHubPlaybackController::PreparePlayback(ULiveLinkRecording* InLiveLinkRecording)
 {
+	if (RecordingPickerWindow)
+	{
+		// If we spawned a recording picker window, let's close it.
+		RecordingPickerWindow->RequestDestroyWindow();
+
+		// Set focus on playback widget so player can press spacebar to play right away.
+		if (TSharedPtr<SWidget> WidgetPtr = PlaybackWidget.Pin())
+		{
+			FSlateApplication::Get().SetUserFocus(0, WidgetPtr);
+		}
+	}
+
 	if (InLiveLinkRecording == nullptr)
 	{
 		UE_LOG(LogLiveLinkHub, Warning, TEXT("Started a recording playback with an invalid recording."));
@@ -485,6 +523,22 @@ void FLiveLinkHubPlaybackController::UnloadRecordingPackage(const TWeakObjectPtr
 	}
 }
 
+void FLiveLinkHubPlaybackController::OnSelectRecording()
+{
+	SAssignNew(RecordingPickerWindow, SWindow)
+		.Title(LOCTEXT("PickRecordingTitle", "Pick a recording."))
+		.ClientSize(FVector2D(750, 440))
+		.SizingRule(ESizingRule::UserSized)
+		[
+			FLiveLinkHub::Get()->GetRecordingListController()->MakeRecordingList()
+		];
+
+	if (GEditor)
+	{
+		GEditor->EditorAddModalWindow(RecordingPickerWindow.ToSharedRef());
+	}
+}
+
 void FLiveLinkHubPlaybackController::GoToTime(FQualifiedFrameTime InTime)
 {
 	// Stop needs to occur to restart playback.
@@ -788,4 +842,6 @@ void FLiveLinkHubPlaybackController::UnpausePlaybackSubjects() const
 	{
 		Client->UnpauseSubject_AnyThread(Subject.SubjectName);
 	}
-}
\ No newline at end of file
+}
+
+#undef LOCTEXT_NAMESPACE
\ No newline at end of file
diff --git a/Engine/Plugins/Animation/LiveLinkHub/Source/LiveLinkHub/Private/Recording/LiveLinkHubPlaybackController.h b/Engine/Plugins/Animation/LiveLinkHub/Source/LiveLinkHub/Private/Recording/LiveLinkHubPlaybackController.h
--- a/Engine/Plugins/Animation/LiveLinkHub/Source/LiveLinkHub/Private/Recording/LiveLinkHubPlaybackController.h
+++ b/Engine/Plugins/Animation/LiveLinkHub/Source/LiveLinkHub/Private/Recording/LiveLinkHubPlaybackController.h
@@ -31,8 +31,8 @@ public:
 	FLiveLinkHubPlaybackController();
 	virtual ~FLiveLinkHubPlaybackController() override;
 
-	/** Create the playback widget. */
-	TSharedRef<SWidget> MakePlaybackWidget();
+	/** Get the playback widget and create it if it doesn't exists. */
+	TSharedRef<SWidget> GetPlaybackWidget();
 
 	/** Apply the recording's preset then prepare the data needed to playback. */
 	void PreparePlayback(ULiveLinkRecording* InLiveLinkRecording);
@@ -78,6 +78,9 @@ public:
 	 */
 	void UnloadRecordingPackage(const TWeakObjectPtr<UPackage>& InPackage, bool bUnloadNextTick = false);
 	
+	/** Create a pop up window that shows available recordings to play. */
+	void OnSelectRecording();
+
 	/** Go to a specific time. */
 	void GoToTime(FQualifiedFrameTime InTime);
 	
@@ -105,7 +108,6 @@ public:
 	/** Retrieve the current framerate. */
 	FFrameRate GetFrameRate() const;
 
-
 	/** Retrieve buffered frame range. */
 	UE::LiveLinkHub::RangeHelpers::Private::TRangeArray<int32> GetBufferedFrameRanges() const;
 	
@@ -165,6 +167,12 @@ public:
 	{
 		return PlaybackReadyDelegate;
 	}
+
+	/** Set whether the playback tab should be closed when the eject button is pressed. */
+	void SetClosePlaybackTabOnEject(bool bCloseOnEject)
+	{
+		bClosePlaybackTabOnEject = bCloseOnEject;
+	}
 	
 	/** Create the playback thread. */
 	void Start();
@@ -227,6 +235,8 @@ private:
 	void UnpausePlaybackSubjects() const;
 
 private:
+	/** Cached pointer to the playback widget. */
+	TWeakPtr<SWidget> PlaybackWidget;
 	/** If the system has established a connection with the client. */
 	bool bIsReady = false;
 	/** Flag for terminating the thread loop */
@@ -284,6 +294,12 @@ private:
 	/** The last sent static frame index. Used to ensure we only send static data when necessary. */
 	TMap<FLiveLinkSubjectKey, int32> LastStaticFrameIndex;
 
+	/** Window that displays the list of recordings that can be played back. */
+	TSharedPtr<class SWindow> RecordingPickerWindow;
+
 	/** Keeps track of whether the recording was paused when starting to scrub to ensure we don't pause if we scrub while playing. */
 	bool bStartedScrubbingPaused = false;
+
+	/** Whether to close the playback tab when ejecting a recording. */
+	bool bClosePlaybackTabOnEject = false;
 };
diff --git a/Engine/Plugins/Animation/LiveLinkHub/Source/LiveLinkHub/Private/UI/Widgets/SLiveLinkHubPlaybackWidget.cpp b/Engine/Plugins/Animation/LiveLinkHub/Source/LiveLinkHub/Private/UI/Widgets/SLiveLinkHubPlaybackWidget.cpp
--- a/Engine/Plugins/Animation/LiveLinkHub/Source/LiveLinkHub/Private/UI/Widgets/SLiveLinkHubPlaybackWidget.cpp
+++ b/Engine/Plugins/Animation/LiveLinkHub/Source/LiveLinkHub/Private/UI/Widgets/SLiveLinkHubPlaybackWidget.cpp
@@ -13,7 +13,7 @@
 #include "Widgets/Input/SButton.h"
 #include "Widgets/Input/SComboButton.h"
 #include "Widgets/Input/SSpinBox.h"
-
+#include "Widgets/Layout/SWidgetSwitcher.h"
 
 #define LOCTEXT_NAMESPACE "SLiveLinkHubPlaybackWidget"
 
@@ -29,6 +29,7 @@ void SLiveLinkHubPlaybackWidget::Construct(const FArguments& InArgs)
 	OnNextFrameDelegate = InArgs._OnNextFrame;
 
 	OnExitPlaybackDelegate = InArgs._OnExitPlayback;
+	OnSelectRecordingDelegate = InArgs._OnSelectRecording;
 
 	OnGetPausedDelegate = InArgs._IsPaused;
 	OnGetIsInReverseDelegate = InArgs._IsInReverse;
@@ -86,7 +87,7 @@ void SLiveLinkHubPlaybackWidget::Construct(const FArguments& InArgs)
 				SNew(SBox)
 				.IsEnabled(IsPlaybackEnabled)
 				[
-					SNew(SLiveLinkHubTimeSlider)
+					SAssignNew(TimeSlider, SLiveLinkHubTimeSlider)
 					.IsEnabled(IsPlaybackEnabled)
 					.BaseArgs(SSimpleTimeSlider::FArguments()
 					.ClampRangeHighlightSize(0.35f)
@@ -235,17 +236,37 @@ void SLiveLinkHubPlaybackWidget::Construct(const FArguments& InArgs)
 				.Padding(FMargin(12.0, 0.0, 0.0, 0.0))
 				.AutoWidth()
 				[
-					SNew(SButton)
-					.ButtonStyle(FAppStyle::Get(), "Animation.PlayControlsButton")
-					.ContentPadding(FMargin(2.0f))
-					.OnClicked(this, &SLiveLinkHubPlaybackWidget::OnExitPlaybackPressed)
-					.ToolTipText(LOCTEXT("EjectButtonToolTip", "Exit Playback"))
+					SNew(SWidgetSwitcher)
+					.WidgetIndex_Lambda([IsPlaybackEnabled]() { return IsPlaybackEnabled.Get() ? 0 : 1; })
+					+ SWidgetSwitcher::Slot()
 					[
-						SNew(SImage)
-							.Image(FAppStyle::Get().GetBrush("LevelViewport.EjectActorPilot.Small"))
+						SNew(SButton)
+						.ButtonStyle(FAppStyle::Get(), "Animation.PlayControlsButton")
+						.ContentPadding(FMargin(2.0f))
+						.IsEnabled(IsPlaybackEnabled)
+						.OnClicked(this, &SLiveLinkHubPlaybackWidget::OnExitPlaybackPressed)
+						.ToolTipText(LOCTEXT("EjectButtonToolTip", "Exit Playback"))
+						[
+							SNew(SImage)
+								.Image(FAppStyle::Get().GetBrush("LevelViewport.EjectActorPilot.Small"))
+								.ColorAndOpacity(FSlateColor::UseSubduedForeground())
+						]
+					]
+					+ SWidgetSwitcher::Slot()
+					[
+						SNew(SButton)
+						.ButtonStyle(FAppStyle::Get(), "Animation.PlayControlsButton")
+						.ContentPadding(FMargin(2.0f))
+						.OnClicked(this, &SLiveLinkHubPlaybackWidget::OnSelectRecordingPressed)
+						.ToolTipText(LOCTEXT("SelectRecordingToolTip", "Select a recording."))
+						[
+							SNew(SImage)
+							.Image(LiveLinkStyle->GetBrush(TEXT("LiveLinkHub.Playback.Icon")))
 							.ColorAndOpacity(FSlateColor::UseSubduedForeground())
+						]
 					]
 				]
+
 				+SHorizontalBox::Slot()
 				.Padding(12.f, 0.f, 0.f, 0.f)
 				[
@@ -330,53 +351,155 @@ void SLiveLinkHubPlaybackWidget::Construct(const FArguments& InArgs)
 	];
 }
 
+void SLiveLinkHubPlaybackWidget::Tick(const FGeometry& AllottedGeometry, const double InCurrentTime, const float InDeltaTime)
+{
+	if (!IsPlaybackEnabled())
+	{
+		return;
+	}
+
+	if (KeyDown[(uint8)EScrubDirection::Left])
+	{
+		KeyDownCounter[(uint8)EScrubDirection::Left]++;
+	}
+
+	if (KeyDown[(uint8)EScrubDirection::Right])
+	{
+		KeyDownCounter[(uint8)EScrubDirection::Right]++;
+	}
+
+	if (KeyDownCounter[(uint8)EScrubDirection::Left] > MinFramesToBeginScrubbing)
+	{
+		// This controls the rate of play when scrubbing.
+		if (KeyDownScrubCounter++ >= FramesBetweenKeyDownScrub)
+		{
+			OnPreviousFramePressed();
+			KeyDownScrubCounter = 0;
+		}
+	}
+	else if (KeyDownCounter[(uint8)EScrubDirection::Right] > MinFramesToBeginScrubbing)
+	{
+		// This controls the rate of play when scrubbing.
+		if (KeyDownScrubCounter++ >= FramesBetweenKeyDownScrub)
+		{
+			OnNextFramePressed();
+			KeyDownScrubCounter = 0;
+		}
+	}
+}
+
+FReply SLiveLinkHubPlaybackWidget::OnFocusReceived(const FGeometry& MyGeometry, const FFocusEvent& InFocusEvent)
+{
+	// Forward focus to the time slider widget to allow 
+	if (InFocusEvent.GetCause() == EFocusCause::SetDirectly)
+	{
+		FSlateApplication::Get().SetUserFocus(0, TimeSlider);
+	}
+
+	return FReply::Handled();
+}
+
+void SLiveLinkHubPlaybackWidget::OnFocusLost(const FFocusEvent& InFocusEvent)
+{
+	KeyDown[(uint8)EScrubDirection::Left] = false;
+	KeyDown[(uint8)EScrubDirection::Right] = false;
+
+	KeyDownCounter[(uint8)EScrubDirection::Left] = 0;
+	KeyDownCounter[(uint8)EScrubDirection::Right] = 0;
+	KeyDownScrubCounter = 0;
+}
+
+FReply SLiveLinkHubPlaybackWidget::OnKeyDown(const FGeometry& MyGeometry, const FKeyEvent& InKeyEvent)
+{
+	if (InKeyEvent.GetKey() == EKeys::Left || InKeyEvent.GetKey() == EKeys::BackSpace)
+	{
+		KeyDown[(uint8)EScrubDirection::Left] = true;
+	}
+	else if (InKeyEvent.GetKey() == EKeys::Right)
+	{
+		KeyDown[(uint8)EScrubDirection::Right] = true;
+	}
+
+	return FReply::Unhandled();
+}
+
+FReply SLiveLinkHubPlaybackWidget::OnKeyUp(const FGeometry& MyGeometry, const FKeyEvent& InKeyEvent)
+{
+	if (InKeyEvent.GetKey() == EKeys::Left || InKeyEvent.GetKey() == EKeys::BackSpace)
+	{
+		KeyDown[(uint8)EScrubDirection::Left] = false;
+		KeyDownCounter[(uint8)EScrubDirection::Left] = 0;
+		KeyDownScrubCounter = 0;
+		return OnPreviousFramePressed().SetUserFocus(AsShared()); // Set focus to prevent navigating on buttons
+	}
+	else if (InKeyEvent.GetKey() == EKeys::Right)
+	{
+		KeyDown[(uint8)EScrubDirection::Right] = false;
+		KeyDownCounter[(uint8)EScrubDirection::Right] = 0;
+		KeyDownScrubCounter = 0;
+		return OnNextFramePressed().SetUserFocus(AsShared()); // Set focus to prevent navigating on buttons
+	}
+	else if (InKeyEvent.GetKey() == EKeys::SpaceBar || InKeyEvent.GetKey() == EKeys::Enter)
+	{
+		OnPlayForwardPressed();
+	}
+
+	return FReply::Unhandled();
+}
+
 FReply SLiveLinkHubPlaybackWidget::OnPlayForwardPressed()
 {
 	OnPlayForwardDelegate.ExecuteIfBound();
-	return FReply::Handled();
+	return HandleAndFocusTimeSlider();
 }
 
 FReply SLiveLinkHubPlaybackWidget::OnPlayReversePressed()
 {
 	OnPlayReverseDelegate.ExecuteIfBound();
-	return FReply::Handled();
+	return HandleAndFocusTimeSlider();
 }
 
 FReply SLiveLinkHubPlaybackWidget::OnFirstFramePressed()
 {
 	OnFirstFrameDelegate.ExecuteIfBound();
-	return FReply::Handled();
+	return HandleAndFocusTimeSlider();
 }
 
 FReply SLiveLinkHubPlaybackWidget::OnLastFramePressed()
 {
 	OnLastFrameDelegate.ExecuteIfBound();
-	return FReply::Handled();
+	return HandleAndFocusTimeSlider();
 }
 
 FReply SLiveLinkHubPlaybackWidget::OnPreviousFramePressed()
 {
 	OnPreviousFrameDelegate.ExecuteIfBound();
-	return FReply::Handled();
+	return HandleAndFocusTimeSlider();
 }
 
 FReply SLiveLinkHubPlaybackWidget::OnNextFramePressed()
 {
 	OnNextFrameDelegate.ExecuteIfBound();
-	return FReply::Handled();
+	return HandleAndFocusTimeSlider();
 }
 
 FReply SLiveLinkHubPlaybackWidget::OnLoopPressed()
 {
 	check(OnSetLoopingDelegate.IsBound() && OnGetLoopingDelegate.IsBound());
 	OnSetLoopingDelegate.Execute(!OnGetLoopingDelegate.Execute());
-	return FReply::Handled();
+	return HandleAndFocusTimeSlider();
 }
 
 FReply SLiveLinkHubPlaybackWidget::OnExitPlaybackPressed()
 {
 	OnExitPlaybackDelegate.ExecuteIfBound();
-	return FReply::Handled();
+	return HandleAndFocusTimeSlider();
+}
+
+FReply SLiveLinkHubPlaybackWidget::OnSelectRecordingPressed()
+{
+	OnSelectRecordingDelegate.ExecuteIfBound();
+	return HandleAndFocusTimeSlider();
 }
 
 void SLiveLinkHubPlaybackWidget::SetCurrentTime(double InTime)
@@ -672,4 +795,9 @@ bool SLiveLinkHubPlaybackWidget::IsPlaybackEnabled() const
 	return IsPlaybackEnabledDelegate.Execute();
 }
 
+FReply SLiveLinkHubPlaybackWidget::HandleAndFocusTimeSlider() const
+{
+	return FReply::Handled().SetUserFocus(TimeSlider.ToSharedRef());
+}
+
 #undef LOCTEXT_NAMESPACE
diff --git a/Engine/Plugins/Animation/LiveLinkHub/Source/LiveLinkHub/Private/UI/Widgets/SLiveLinkHubPlaybackWidget.h b/Engine/Plugins/Animation/LiveLinkHub/Source/LiveLinkHub/Private/UI/Widgets/SLiveLinkHubPlaybackWidget.h
--- a/Engine/Plugins/Animation/LiveLinkHub/Source/LiveLinkHub/Private/UI/Widgets/SLiveLinkHubPlaybackWidget.h
+++ b/Engine/Plugins/Animation/LiveLinkHub/Source/LiveLinkHub/Private/UI/Widgets/SLiveLinkHubPlaybackWidget.h
@@ -45,6 +45,8 @@ public:
 	SLATE_EVENT(FOnButtonPressed, OnNextFrame)
 	/** When the eject button is pressed. */
 	SLATE_EVENT(FOnButtonPressed, OnExitPlayback)
+	/** When the select recording button is pressed. */
+	SLATE_EVENT(FOnButtonPressed, OnSelectRecording)
 	/** Is the recording looping? */
 	SLATE_EVENT(FOnGetChecked, IsLooping)
 	/** Set the loop state of the recording. */
@@ -95,6 +97,16 @@ public:
 	 */
 	void Construct(const FArguments& InArgs);
 
+protected:
+	//~ Begin SWidget interface
+	virtual bool SupportsKeyboardFocus() const { return true; }
+	virtual FReply OnFocusReceived(const FGeometry& MyGeometry, const FFocusEvent& InFocusEvent) override;
+	virtual void OnFocusLost(const FFocusEvent& InFocusEvent) override;
+	virtual void Tick(const FGeometry& AllottedGeometry, const double InCurrentTime, const float InDeltaTime);
+	virtual FReply OnKeyDown(const FGeometry& MyGeometry, const FKeyEvent& InKeyEvent) override;
+	virtual FReply OnKeyUp(const FGeometry& MyGeometry, const FKeyEvent& InKeyEvent) override;
+	//~ End SWidget interface
+
 private:
 	/** When play forward is pressed. */
 	FReply OnPlayForwardPressed();
@@ -110,6 +122,8 @@ private:
 	FReply OnNextFramePressed();
 	/** When loop recording is pressed. */
 	FReply OnLoopPressed();
+	/** When using wishes to select a recording to play. */
+	FReply OnSelectRecordingPressed();
 	/** When recording is exited. */
 	FReply OnExitPlaybackPressed();
 
@@ -191,6 +205,9 @@ private:
 
 	/** Returns whether the playback widgets should be enabled or not by calling the IsPlaybackEnabledDelegate delegate. */
 	bool IsPlaybackEnabled() const;
+
+	/** Returns a FReply::Handled() and focuses the time slider. Useful so the user can click the play button but scrub using left and right keys. */
+	FReply HandleAndFocusTimeSlider() const;
 private:
 	/** Delegate for pressing play forward. */
 	FOnButtonPressed OnPlayForwardDelegate;
@@ -213,6 +230,9 @@ private:
 	/** Delegate for exiting the current recording. */
 	FOnButtonPressed OnExitPlaybackDelegate;
 
+	/** Delegate for selecting a recording to play. */
+	FOnButtonPressed OnSelectRecordingDelegate;
+
 	/** Delegate checking if the recording playback is enabled */
 	FOnGetChecked IsPlaybackEnabledDelegate;
 
@@ -278,4 +298,29 @@ private:
 
 	/** The format to display/edit values in. */
 	EFrameNumberDisplayFormats DisplayFormat = EFrameNumberDisplayFormats::NonDropFrameTimecode;
+
+	/** Holds the LiveLinkHubTimeSlider widget. */
+	TSharedPtr<SWidget> TimeSlider;
+
+	/** Helper enum to track KeyDown state. */
+	enum class EScrubDirection : uint8
+	{
+		Left,
+		Right
+	};
+
+	/** Keeps track of how many frames since the left and right key have been pressed. */
+	uint32 KeyDownCounter[2] = { 0, 0 };
+
+	/** Keeps track of whether the left or right key is being pressed. */
+	bool KeyDown[2] = { 0, 0 };
+
+	/** Counter used to advance the recording while holding down left or right key. */
+	uint8 KeyDownScrubCounter = 0;
+
+	/** Number of frames to wait between advancing the frame while a key is held down.*/
+	static constexpr uint8 FramesBetweenKeyDownScrub = 3;
+
+	/** Minimum number of frames to consider a frame pressed as being held down for scrubbing (vs. tapping it for a single frame jump). */
+	static constexpr uint8 MinFramesToBeginScrubbing = 30;
 };
\ No newline at end of file
