LiveLinkHub - Scan paths when files are added or removed from the content folder

diff --git a/Engine/Plugins/Animation/LiveLinkHub/Source/LiveLinkHub/Private/LiveLinkHub.cpp b/Engine/Plugins/Animation/LiveLinkHub/Source/LiveLinkHub/Private/LiveLinkHub.cpp
--- a/Engine/Plugins/Animation/LiveLinkHub/Source/LiveLinkHub/Private/LiveLinkHub.cpp
+++ b/Engine/Plugins/Animation/LiveLinkHub/Source/LiveLinkHub/Private/LiveLinkHub.cpp
@@ -121,13 +121,6 @@ void FLiveLinkHub::Initialize(FLiveLinkHubTicker& Ticker)
 	FModuleManager::Get().LoadModule("Settings");
 	FModuleManager::Get().LoadModule("StatusBar");
 
-	if (LiveLinkHubInternalUtils::IsDistributedBuild())
-	{
-		IAssetRegistry& AssetRegistry = FModuleManager::LoadModuleChecked<FAssetRegistryModule>(AssetRegistryConstants::ModuleName).Get();
-		const FString FilePath = FPaths::Combine(FPlatformProcess::UserSettingsDir(), *FApp::GetEpicProductIdentifier(), TEXT("LiveLinkHub"), TEXT("Content"));
-		AssetRegistry.ScanPathsSynchronous({ FilePath }, /*bForceRescan=*/ true);
-	}
-
 	CommandExecutor = MakeUnique<FConsoleCommandExecutor>();
 	IModularFeatures::Get().RegisterModularFeature(IConsoleCommandExecutor::ModularFeatureName(), CommandExecutor.Get());
 
@@ -537,33 +530,35 @@ void FLiveLinkHub::OnSubjectOutboundNameModified(const FLiveLinkSubjectKey& Subj
 
 void FLiveLinkHub::OnPlaybackReadyForUpgrade()
 {
-	if (GetPlaybackController() && GetPlaybackController()->GetRecording().IsValid())
+	if (!PlaybackController || !RecordingController || !PlaybackController->GetRecording())
 	{
-		GetPlaybackController()->OnPlaybackReady().RemoveAll(this);
+		return;
+	}
 
-		// Get the save path before starting the long task.
-		const FString DefaultName = GetPlaybackController()->GetRecording()->GetName() + TEXT("_copy");
-		FString RecordingPath;
-		if (!GetRecordingController()->OpenSaveAsPrompt(RecordingPath, DefaultName))
-		{
-			GetPlaybackController()->EjectAndUnload();
-			return;
-		}
+	PlaybackController->OnPlaybackReady().RemoveAll(this);
+
+	// Get the save path before starting the long task.
+	const FString DefaultName = PlaybackController->GetRecording()->GetName() + TEXT("_copy");
+	FString RecordingPath;
+	if (!RecordingController->OpenSaveAsPrompt(RecordingPath, DefaultName))
+	{
+		PlaybackController->EjectAndUnload();
+		return;
+	}
 		
-		ULiveLinkUAssetRecording* CurrentRecording = CastChecked<ULiveLinkUAssetRecording>(GetPlaybackController()->GetRecording().Get());
+	ULiveLinkUAssetRecording* CurrentRecording = CastChecked<ULiveLinkUAssetRecording>(PlaybackController->GetRecording().Get());
 		
-		TWeakObjectPtr<ULiveLinkUAssetRecording> CurrentRecordingWeakPtr(CurrentRecording);
-		CurrentRecording->GetOnRecordingFullyLoaded().BindLambda([this, CurrentRecordingWeakPtr, RecordingPath]()
+	TWeakObjectPtr<ULiveLinkUAssetRecording> CurrentRecordingWeakPtr(CurrentRecording);
+	CurrentRecording->GetOnRecordingFullyLoaded().BindLambda([this, CurrentRecordingWeakPtr, RecordingPath]()
+	{
+		const TStrongObjectPtr<ULiveLinkUAssetRecording> CurrentRecordingPin = CurrentRecordingWeakPtr.Pin();
+		if (CurrentRecordingPin.IsValid() && RecordingController)
 		{
-			const TStrongObjectPtr<ULiveLinkUAssetRecording> CurrentRecordingPin = CurrentRecordingWeakPtr.Pin();
-			if (CurrentRecordingPin.IsValid() && GetRecordingController())
-			{
-				GetRecordingController()->RecordFromExistingAndSave(CurrentRecordingPin.Get(), RecordingPath);
-			}
-		});
+			RecordingController->RecordFromExistingAndSave(CurrentRecordingPin.Get(), RecordingPath);
+		}
+	});
 
-		CurrentRecording->LoadEntireRecordingForUpgrade();
-	}
+	CurrentRecording->LoadEntireRecordingForUpgrade();
 }
 
 void FLiveLinkHub::RegisterLiveLinkHubSettings()
diff --git a/Engine/Plugins/Animation/LiveLinkHub/Source/LiveLinkHub/Private/Recording/LiveLinkHubRecordingListController.h b/Engine/Plugins/Animation/LiveLinkHub/Source/LiveLinkHub/Private/Recording/LiveLinkHubRecordingListController.h
--- a/Engine/Plugins/Animation/LiveLinkHub/Source/LiveLinkHub/Private/Recording/LiveLinkHubRecordingListController.h
+++ b/Engine/Plugins/Animation/LiveLinkHub/Source/LiveLinkHub/Private/Recording/LiveLinkHubRecordingListController.h
@@ -4,7 +4,9 @@
 
 #include "AssetRegistry/AssetData.h"
 #include "IContentBrowserSingleton.h"
+#include "DirectoryWatcherModule.h"
 #include "Framework/MultiBox/MultiBoxBuilder.h"
+#include "IDirectoryWatcher.h"
 #include "LiveLinkHub.h"
 #include "LiveLinkHubLog.h"
 #include "LiveLinkHubPlaybackController.h"
@@ -22,6 +24,35 @@ public:
 	FLiveLinkHubRecordingListController(const TSharedRef<FLiveLinkHub>& InLiveLinkHub)
 		: LiveLinkHub(InLiveLinkHub)
 	{
+		TStringBuilder<128> MountPackageName;
+		TStringBuilder<128> MountFilePath;
+		TStringBuilder<128> RelPath;
+
+		if (!FPackageName::TryGetMountPointForPath(TEXT("/Game"), MountPackageName, MountFilePath, RelPath))
+		{
+			return;
+		}
+
+		ContentDirectory = MountFilePath.ToString();
+
+		IAssetRegistry& AssetRegistry = FModuleManager::LoadModuleChecked<FAssetRegistryModule>(AssetRegistryConstants::ModuleName).Get();
+		AssetRegistry.ScanPathsSynchronous({ ContentDirectory }, /*bForceRescan=*/ true);
+
+		// Add a directory watcher callback to refresh recordings.
+		if (IDirectoryWatcher* Watcher = FModuleManager::LoadModuleChecked<FDirectoryWatcherModule>(TEXT("DirectoryWatcher")).Get())
+		{
+			Watcher->RegisterDirectoryChangedCallback_Handle(ContentDirectory, IDirectoryWatcher::FDirectoryChanged::CreateRaw(this, &FLiveLinkHubRecordingListController::OnContentDirectoryChanged), DirectoryChangedHandle, IDirectoryWatcher::IncludeDirectoryChanges);
+		}
+	}
+
+	~FLiveLinkHubRecordingListController()
+	{
+		FDirectoryWatcherModule* DirectoryWatcher = FModuleManager::GetModulePtr<FDirectoryWatcherModule>(TEXT("DirectoryWatcher"));
+
+		if (DirectoryWatcher && DirectoryWatcher->Get())
+		{
+			DirectoryWatcher->Get()->UnregisterDirectoryChangedCallback_Handle(ContentDirectory, DirectoryChangedHandle);
+		}
 	}
 
 	/** Create the list's widget. */
@@ -62,10 +93,24 @@ private:
 			HubPtr->GetPlaybackController()->PreparePlayback(ImportedRecording);
 		}
 	}
+	/** Reacts to changes in the content directory to rescan it to update the recording list. */
+	void OnContentDirectoryChanged(const TArray<struct FFileChangeData>& FileChangeData)
+	{
+		TArray<FString> FilePaths;
+		FilePaths.Reserve(FileChangeData.Num());
+
+		Algo::Transform(FileChangeData, FilePaths, [](const FFileChangeData& Data) { return Data.Filename; });
+		IAssetRegistry& AssetRegistry = FModuleManager::LoadModuleChecked<FAssetRegistryModule>(AssetRegistryConstants::ModuleName).Get();
+		AssetRegistry.ScanPathsSynchronous({ FilePaths }, /*bForceRescan=*/ true);
+	}
 
 private:
 	/** LiveLinkHub object that holds the different controllers. */
 	TWeakPtr<FLiveLinkHub> LiveLinkHub;
+	/** Delegate handle for the RegisterDirectoryChangedCallback_Handle callback. */
+	FDelegateHandle DirectoryChangedHandle;
+	/** Path to the content directory. Used to react to changes to that directory, such as new files being added. */
+	FString ContentDirectory;
 };
 
 #undef LOCTEXT_NAMESPACE /* LiveLinkHub.RecordingList */
