LiveLinkHub - Fix playback window size being too small when using high DPI

diff --git a/Engine/Plugins/Animation/LiveLinkHub/Source/LiveLinkHub/Private/LiveLinkHubPlaybackAppMode.h b/Engine/Plugins/Animation/LiveLinkHub/Source/LiveLinkHub/Private/LiveLinkHubPlaybackAppMode.h
--- a/Engine/Plugins/Animation/LiveLinkHub/Source/LiveLinkHub/Private/LiveLinkHubPlaybackAppMode.h
+++ b/Engine/Plugins/Animation/LiveLinkHub/Source/LiveLinkHub/Private/LiveLinkHubPlaybackAppMode.h
@@ -4,6 +4,7 @@
 
 #include "LiveLinkHubApplicationMode.h"
 
+#include "HAL/PlatformApplicationMisc.h"
 #include "Framework/Docking/TabManager.h"
 #include "LiveLinkHubApplicationBase.h"
 #include "LiveLinkHub.h"
@@ -44,9 +45,13 @@ public:
 		{
 			PlaybackController->SetClosePlaybackTabOnEject(false);
 			
-			PreviousWindowSize = FLiveLinkHub::Get()->GetRootWindow()->GetSizeInScreen();
+			TSharedRef<SWindow> RootWindow = FLiveLinkHub::Get()->GetRootWindow();
 
-			FLiveLinkHub::Get()->GetRootWindow()->Resize({ 800.0, 220.0 });
+			PreviousWindowSize = RootWindow->GetSizeInScreen();
+			FVector2D WindowPosition = RootWindow->GetPositionInScreen();
+			FVector2D TargetSize = FVector2D{ 800.0, 220.0 } * FPlatformApplicationMisc::GetDPIScaleFactorAtPoint(WindowPosition.X, WindowPosition.Y);
+
+			RootWindow->Resize(TargetSize);
 		}
 	}
 
