LiveLinkHub - Fix CIS compile errors caused by using FReadScopeLock on FTransactionallySafeRWLock

diff --git a/Engine/Plugins/Animation/LiveLinkHub/Source/LiveLinkHub/Private/Clients/LiveLinkHubProvider.cpp b/Engine/Plugins/Animation/LiveLinkHub/Source/LiveLinkHub/Private/Clients/LiveLinkHubProvider.cpp
--- a/Engine/Plugins/Animation/LiveLinkHub/Source/LiveLinkHub/Private/Clients/LiveLinkHubProvider.cpp
+++ b/Engine/Plugins/Animation/LiveLinkHub/Source/LiveLinkHub/Private/Clients/LiveLinkHubProvider.cpp
@@ -114,7 +114,7 @@ bool FLiveLinkHubProvider::ShouldTransmitToSubject_AnyThread(FName SubjectName,
 
 	FLiveLinkHubClientId ClientId;
 	{
-		FReadScopeLock Locker(ClientsMapLock);
+		UE::TReadScopeLock Locker(ClientsMapLock);
 		ClientId = AddressToIdCache.FindRef(Address);
 	}
 
@@ -152,7 +152,7 @@ void FLiveLinkHubProvider::DisconnectAll()
 
 	TArray<FLiveLinkHubAddressBookEntry> ClientsToDisconnect;
 	{
-		FReadScopeLock Lock(AddressBookLock);
+		UE::TReadScopeLock Lock(AddressBookLock);
 		ClientsToDisconnect = AddressBook.FilterByPredicate([](const FLiveLinkHubAddressBookEntry& Entry) { return Entry.bConnected; });
 	}
 
@@ -172,7 +172,7 @@ void FLiveLinkHubProvider::DisconnectClient(const FLiveLinkHubClientId& Client)
 
 	TOptional<FLiveLinkHubAddressBookEntry> EntryToDisconnect;
 	{
-		FReadScopeLock Lock(AddressBookLock);
+		UE::TReadScopeLock Lock(AddressBookLock);
 		if (FLiveLinkHubAddressBookEntry* FoundEntry = Algo::FindBy(AddressBook, Client, &FLiveLinkHubAddressBookEntry::ClientId))
 		{
 			EntryToDisconnect = *FoundEntry;
@@ -186,7 +186,7 @@ void FLiveLinkHubProvider::DisconnectClient(const FLiveLinkHubClientId& Client)
 	else
 	{
 		// This might have been a restored client that hasn't been updated, so clear it. 
-		FWriteScopeLock Locker(ClientsMapLock);
+		UE::TWriteScopeLock Locker(ClientsMapLock);
 		ClientsMap.Remove(Client);
 
 		for (auto It = AddressToIdCache.CreateIterator(); It; ++It)
@@ -209,7 +209,7 @@ void FLiveLinkHubProvider::DisconnectClient(const FLiveLinkHubAddressBookEntry&
 	// Then modify the entry in our address book.
 	TOptional<FLiveLinkHubAddressBookEntry> AddressBookEntry;
 	{
-		FWriteScopeLock Lock(AddressBookLock);
+		UE::TWriteScopeLock Lock(AddressBookLock);
 		if (FLiveLinkHubAddressBookEntry* EntryPtr = AddressBook.FindByPredicate([&EntryToDisconnect](const FLiveLinkHubAddressBookEntry& Entry) { return Entry.ClientId == EntryToDisconnect.ClientId; }))
 		{
 			EntryPtr->bConnected = false;
@@ -236,7 +236,7 @@ void FLiveLinkHubProvider::SendTimecodeSettings(const FLiveLinkHubTimecodeSettin
 		FMessageAddress RelevantControlEndpoint;
 
 		{
-			FReadScopeLock Lock(AddressBookLock);
+			UE::TReadScopeLock Lock(AddressBookLock);
 			if (const FLiveLinkHubAddressBookEntry* Entry = Algo::FindBy(AddressBook, ClientId, &FLiveLinkHubAddressBookEntry::ClientId))
 			{
 				if (Entry->bConnected)
@@ -270,7 +270,7 @@ void FLiveLinkHubProvider::SendCustomTimeStepSettings(const FLiveLinkHubCustomTi
 		FMessageAddress RelevantControlEndpoint;
 
 		{
-			FReadScopeLock Lock(AddressBookLock);
+			UE::TReadScopeLock Lock(AddressBookLock);
 			if (const FLiveLinkHubAddressBookEntry* Entry = Algo::FindBy(AddressBook, ClientId, &FLiveLinkHubAddressBookEntry::ClientId))
 			{
 				if (Entry->bConnected)
@@ -308,7 +308,7 @@ void FLiveLinkHubProvider::AddRestoredClient(FLiveLinkHubUEClientInfo& RestoredC
 			TOptional<FLiveLinkHubAddressBookEntry> FoundClient;
 			{
 				// todo: validate this logic for a client that's already connected.
-				FWriteScopeLock Lock(AddressBookLock);
+				UE::TWriteScopeLock Lock(AddressBookLock);
 				for (FLiveLinkHubAddressBookEntry& Entry : AddressBook)
 				{
 					// todo: Multiple clients with same hostname, discriminate by project
@@ -323,7 +323,7 @@ void FLiveLinkHubProvider::AddRestoredClient(FLiveLinkHubUEClientInfo& RestoredC
 
 			if (FoundClient)
 			{
-				FWriteScopeLock Locker(ClientsMapLock);
+				UE::TWriteScopeLock Locker(ClientsMapLock);
 				ClientsMap.FindOrAdd(RestoredClientInfo.Id) = RestoredClientInfo;
 
 				SendMessage(MakeLiveLinkHubDiscoveryMessage(), FoundClient->ControlEndpoint, EMessageFlags::Reliable);
@@ -332,7 +332,7 @@ void FLiveLinkHubProvider::AddRestoredClient(FLiveLinkHubUEClientInfo& RestoredC
 			else
 			{
 				// If we haven't found a v2 client for this entry, fallback to previous method.
-				FReadScopeLock Locker(ClientsMapLock);
+				UE::TReadScopeLock Locker(ClientsMapLock);
 
 				for (auto It = ClientsMap.CreateIterator(); It; ++It)
 				{
@@ -362,7 +362,7 @@ void FLiveLinkHubProvider::AddRestoredClient(FLiveLinkHubUEClientInfo& RestoredC
 
 TOptional<FLiveLinkHubUEClientInfo> FLiveLinkHubProvider::GetClientInfo(FLiveLinkHubClientId InClient) const
 {
-	FReadScopeLock Locker(ClientsMapLock);
+	UE::TReadScopeLock Locker(ClientsMapLock);
 	TOptional<FLiveLinkHubUEClientInfo> ClientInfo;
 	if (const FLiveLinkHubUEClientInfo* ClientInfoPtr = ClientsMap.Find(InClient))
 	{
@@ -406,7 +406,7 @@ void FLiveLinkHubProvider::HandleHubConnectMessage(const FLiveLinkHubConnectMess
 
 	TOptional<FLiveLinkHubClientId> UpdatedClient;
 	{
-		FWriteScopeLock Locker(ClientsMapLock);
+		UE::TWriteScopeLock Locker(ClientsMapLock);
 
 		// First check if there are multiple disconnected entries with the same host as the incoming client.
 		uint8 NumClientsForHost = 0;
@@ -452,7 +452,7 @@ void FLiveLinkHubProvider::HandleHubConnectMessage(const FLiveLinkHubConnectMess
 	{
 		// A client was just restored, make sure to update its address book entry.
 		{
-			FWriteScopeLock Lock(AddressBookLock);
+			UE::TWriteScopeLock Lock(AddressBookLock);
 			if (FLiveLinkHubAddressBookEntry* Entry = Algo::FindBy(AddressBook, ControlEndpoint, &FLiveLinkHubAddressBookEntry::ControlEndpoint))
 			{
 				Entry->bConnected = true;
@@ -470,7 +470,7 @@ void FLiveLinkHubProvider::HandleHubConnectMessage(const FLiveLinkHubConnectMess
 		FLiveLinkHubClientId ClientId;
 
 		{
-			FWriteScopeLock Lock(AddressBookLock);
+			UE::TWriteScopeLock Lock(AddressBookLock);
 
 			if (FLiveLinkHubAddressBookEntry* Entry = Algo::FindBy(AddressBook, ControlEndpoint, &FLiveLinkHubAddressBookEntry::ControlEndpoint))
 			{
@@ -497,7 +497,7 @@ void FLiveLinkHubProvider::HandleHubConnectMessage(const FLiveLinkHubConnectMess
 		UpdatedClient = ClientId; // Set this so we can update the client's timecode provider down below.
 
 		{
-			FWriteScopeLock Locker(ClientsMapLock);
+			UE::TWriteScopeLock Locker(ClientsMapLock);
 			AddressToIdCache.FindOrAdd(DataEndpoint) = ClientId;
 			ClientsMap.Add(NewClient.Id, MoveTemp(NewClient));
 		}
@@ -545,7 +545,7 @@ void FLiveLinkHubProvider::HandleClientInfoMessage(const FLiveLinkClientInfoMess
 
 	FLiveLinkHubClientId ClientId;
 	{
-		FReadScopeLock Lock(AddressBookLock);
+		UE::TReadScopeLock Lock(AddressBookLock);
 
 		if (const FLiveLinkHubAddressBookEntry* Entry = Algo::FindBy(AddressBook, Context->GetSender(), &FLiveLinkHubAddressBookEntry::ControlEndpoint))
 		{
@@ -554,7 +554,7 @@ void FLiveLinkHubProvider::HandleClientInfoMessage(const FLiveLinkClientInfoMess
 	}
 
 	{
-		FWriteScopeLock Locker(ClientsMapLock);
+		UE::TWriteScopeLock Locker(ClientsMapLock);
 		if (FLiveLinkHubUEClientInfo* ClientInfo = ClientsMap.Find(ClientId))
 		{
 			ClientInfo->UpdateFromInfoMessage(Message);
@@ -586,7 +586,7 @@ void FLiveLinkHubProvider::HandleHubDisconnectMessage(const FLiveLinkHubDisconne
 
 	FLiveLinkHubClientId ClientId;
 	{
-		FWriteScopeLock Lock(AddressBookLock);
+		UE::TWriteScopeLock Lock(AddressBookLock);
 		if (FLiveLinkHubAddressBookEntry* AddressBookEntry = AddressBook.FindByPredicate([&Address](const FLiveLinkHubAddressBookEntry& Entry) { return Entry.ControlEndpoint == Address; }))
 		{
 			AddressBookEntry->bConnected = false;
@@ -621,7 +621,7 @@ bool FLiveLinkHubProvider::ShouldTransmitToClient_AnyThread(const FLiveLinkHubCl
 		return false;
 	}
 
-	FReadScopeLock Locker(ClientsMapLock);
+	UE::TReadScopeLock Locker(ClientsMapLock);
 	if (const FLiveLinkHubUEClientInfo* ClientInfoPtr = ClientsMap.Find(ClientId))
 	{
 		if (const TSharedPtr<ILiveLinkHubSessionManager> Manager = SessionManager.Pin())
@@ -646,7 +646,7 @@ void FLiveLinkHubProvider::ConnectToAllDiscoveredClients()
 	TArray<FMessageAddress> CompatibleClients;
 
 	{
-		FReadScopeLock Lock(AddressBookLock);
+		UE::TReadScopeLock Lock(AddressBookLock);
 		Algo::TransformIf(AddressBook, CompatibleClients,
 			[](const FLiveLinkHubAddressBookEntry& Entry) { return UE::LiveLink::Messaging::CanTransmitTo(FLiveLinkHub::Get()->GetTopologyMode(), Entry.TopologyMode); },
 			[](const FLiveLinkHubAddressBookEntry& Entry) { return Entry.ControlEndpoint; });
@@ -664,7 +664,7 @@ void FLiveLinkHubProvider::OnConnectionsClosed(const TArray<FMessageAddress>& Cl
 
 	for (FMessageAddress TrackedAddress : ClosedAddresses)
 	{
-		FWriteScopeLock Locker(ClientsMapLock);
+		UE::TWriteScopeLock Locker(ClientsMapLock);
 		if (FLiveLinkHubClientId* ClientId = AddressToIdCache.Find(TrackedAddress))
 		{
 			// Removing this might have implications for restoring sessions.
@@ -682,7 +682,7 @@ void FLiveLinkHubProvider::CloseConnections(const TArray<FMessageAddress>& Close
 	TArray<FLiveLinkHubClientId> Notifications;
 	TArray<FMessageAddress> AddressesToRemove;
 	{
-		FWriteScopeLock Locker(ClientsMapLock);
+		UE::TWriteScopeLock Locker(ClientsMapLock);
 
 		for (FMessageAddress TrackedAddress : ClosedAddresses)
 		{
@@ -697,7 +697,7 @@ void FLiveLinkHubProvider::CloseConnections(const TArray<FMessageAddress>& Close
 	}
 
 	{
-		FWriteScopeLock Lock(AddressBookLock);
+		UE::TWriteScopeLock Lock(AddressBookLock);
 		for (const FMessageAddress& Address : AddressesToRemove)
 		{
 			for (auto It = AddressBook.CreateIterator(); It; ++It)
@@ -724,7 +724,7 @@ void FLiveLinkHubProvider::ValidateHubConnections()
 	const double MessageBusTimeout = GetDefault<ULiveLinkSettings>()->GetMessageBusTimeBeforeRemovingDeadSource();
 
 	{
-		FWriteScopeLock Lock(AddressBookLock);
+		UE::TWriteScopeLock Lock(AddressBookLock);
 		for (auto It = AddressBook.CreateIterator(); It; ++It)
 		{
 			if (!It->bConnected)
@@ -834,7 +834,7 @@ TArray<FLiveLinkHubClientId> FLiveLinkHubProvider::GetDiscoveredClients() const
 			}
 		};
 
-		FReadScopeLock Lock(AddressBookLock);
+		UE::TReadScopeLock Lock(AddressBookLock);
 		Algo::TransformIf(AddressBook, DiscoveredClients,
 			CanTransmitTo,
 			[](const FLiveLinkHubAddressBookEntry& Entry) { return Entry.ClientId; });
@@ -845,7 +845,7 @@ TArray<FLiveLinkHubClientId> FLiveLinkHubProvider::GetDiscoveredClients() const
 
 FText FLiveLinkHubProvider::GetClientDisplayName(FLiveLinkHubClientId InId) const
 {
-	FReadScopeLock Locker(ClientsMapLock);
+	UE::TReadScopeLock Locker(ClientsMapLock);
 	FText DisplayName;
 
 	if (const FLiveLinkHubUEClientInfo* ClientInfoPtr = ClientsMap.Find(InId))
@@ -861,7 +861,7 @@ FText FLiveLinkHubProvider::GetClientDisplayName(FLiveLinkHubClientId InId) cons
 	}
 	else
 	{
-		FReadScopeLock Lock(AddressBookLock);
+		UE::TReadScopeLock Lock(AddressBookLock);
 		if (const FLiveLinkHubAddressBookEntry* Entry = AddressBook.FindByPredicate([InId](const FLiveLinkHubAddressBookEntry& Entry) { return Entry.ClientId == InId; }))
 		{
 			DisplayName = FText::FromString(FString::Format(TEXT("{0} ({1})"), { *Entry->Hostname, *Entry->IP }));
@@ -877,7 +877,7 @@ FText FLiveLinkHubProvider::GetClientDisplayName(FLiveLinkHubClientId InId) cons
 
 FText FLiveLinkHubProvider::GetClientStatus(FLiveLinkHubClientId Client) const
 {
-	FReadScopeLock Locker(ClientsMapLock);
+	UE::TReadScopeLock Locker(ClientsMapLock);
 	if (const FLiveLinkHubUEClientInfo* ClientInfoPtr = ClientsMap.Find(Client))
 	{
 		return StaticEnum<ELiveLinkClientStatus>()->GetDisplayNameTextByValue(static_cast<int64>(ClientInfoPtr->Status));
@@ -888,7 +888,7 @@ FText FLiveLinkHubProvider::GetClientStatus(FLiveLinkHubClientId Client) const
 
 bool FLiveLinkHubProvider::IsClientEnabled(FLiveLinkHubClientId Client) const
 {
-	FReadScopeLock Locker(ClientsMapLock);
+	UE::TReadScopeLock Locker(ClientsMapLock);
 	if (const FLiveLinkHubUEClientInfo* ClientInfoPtr = ClientsMap.Find(Client))
 	{
 		return ClientInfoPtr->bEnabled;
@@ -898,7 +898,7 @@ bool FLiveLinkHubProvider::IsClientEnabled(FLiveLinkHubClientId Client) const
 
 bool FLiveLinkHubProvider::IsClientConnected(FLiveLinkHubClientId Client) const
 {
-	FReadScopeLock Locker(ClientsMapLock);
+	UE::TReadScopeLock Locker(ClientsMapLock);
 	if (const FLiveLinkHubUEClientInfo* ClientInfoPtr = ClientsMap.Find(Client))
 	{
 		return ClientInfoPtr->Status == ELiveLinkClientStatus::Connected;
@@ -911,7 +911,7 @@ void FLiveLinkHubProvider::ConnectTo(FLiveLinkHubClientId Client)
 	FLiveLinkHubAddressBookEntry Entry;
 
 	{
-		FWriteScopeLock Lock(AddressBookLock);
+		UE::TWriteScopeLock Lock(AddressBookLock);
 
 		int32 EntryIndex = AddressBook.IndexOfByPredicate([Client](const FLiveLinkHubAddressBookEntry& Entry) { return Entry.ClientId == Client; });
 		if (EntryIndex == INDEX_NONE)
@@ -947,7 +947,7 @@ void FLiveLinkHubProvider::ConnectTo(FLiveLinkHubClientId Client)
 void FLiveLinkHubProvider::SetClientEnabled(FLiveLinkHubClientId Client, bool bInEnable)
 {
 	{
-		FWriteScopeLock Locker(ClientsMapLock);
+		UE::TWriteScopeLock Locker(ClientsMapLock);
 		if (FLiveLinkHubUEClientInfo* ClientInfoPtr = ClientsMap.Find(Client))
 		{
 			ClientInfoPtr->bEnabled = bInEnable;
@@ -986,7 +986,7 @@ void FLiveLinkHubProvider::SetClientEnabled(FLiveLinkHubClientId Client, bool bI
 
 bool FLiveLinkHubProvider::IsSubjectEnabled(FLiveLinkHubClientId Client, FName SubjectName) const
 {
-	FReadScopeLock Locker(ClientsMapLock);
+	UE::TReadScopeLock Locker(ClientsMapLock);
 	if (const FLiveLinkHubUEClientInfo* ClientInfoPtr = ClientsMap.Find(Client))
 	{
 		return !ClientInfoPtr->DisabledSubjects.Contains(SubjectName);
@@ -996,7 +996,7 @@ bool FLiveLinkHubProvider::IsSubjectEnabled(FLiveLinkHubClientId Client, FName S
 
 void FLiveLinkHubProvider::SetSubjectEnabled(FLiveLinkHubClientId Client, FName SubjectName, bool bInEnable)
 {
-	FWriteScopeLock Locker(ClientsMapLock);
+	UE::TWriteScopeLock Locker(ClientsMapLock);
 	if (FLiveLinkHubUEClientInfo* ClientInfoPtr = ClientsMap.Find(Client))
 	{
 		if (bInEnable)
@@ -1017,7 +1017,7 @@ FLiveLinkHubClientId FLiveLinkHubProvider::UpdateAddressBookEntry(const FMessage
 	FLiveLinkHubClientId ClientId;
 
 	{
-		FWriteScopeLock Lock(AddressBookLock);
+		UE::TWriteScopeLock Lock(AddressBookLock);
 		
 		if (FLiveLinkHubAddressBookEntry* Entry = AddressBook.FindByPredicate([Address = ControlEndpointAddress](const FLiveLinkHubAddressBookEntry& Entry) { return Entry.ControlEndpoint == Address; }))
 		{
@@ -1043,7 +1043,7 @@ FLiveLinkHubClientId FLiveLinkHubProvider::UpdateAddressBookEntry(const FMessage
 	{
 		bool bModifiedIP = false;
 		{
-			FWriteScopeLock Lock(ClientsMapLock);
+			UE::TWriteScopeLock Lock(ClientsMapLock);
 			// Update the clients info from the address book entry.
 			// Todo: Ideally, this map and the address book should be one and the same.
 			if (FLiveLinkHubUEClientInfo* Info = ClientsMap.Find(ReadOnlyEntry->ClientId); Info && Info->IPAddress != ReadOnlyEntry->IP)
@@ -1074,7 +1074,7 @@ FLiveLinkHubClientId FLiveLinkHubProvider::UpdateAddressBookEntry(const FMessage
 		};
 
 		{
-			FWriteScopeLock Lock(AddressBookLock);
+			UE::TWriteScopeLock Lock(AddressBookLock);
 			AddressBook.Add(NewEntry);
 		}
 
@@ -1124,7 +1124,7 @@ void FLiveLinkHubProvider::OnMessageBusNotification(const FMessageBusNotificatio
 	{
 		TArray<FLiveLinkHubClientId> RemovedClients;
 
-		FWriteScopeLock Lock(AddressBookLock);
+		UE::TWriteScopeLock Lock(AddressBookLock);
 
 		for (auto EntryIt = AddressBook.CreateIterator(); EntryIt; ++EntryIt)
 		{
diff --git a/Engine/Plugins/Animation/LiveLinkHub/Source/LiveLinkHub/Private/Clients/LiveLinkHubProvider.h b/Engine/Plugins/Animation/LiveLinkHub/Source/LiveLinkHub/Private/Clients/LiveLinkHubProvider.h
--- a/Engine/Plugins/Animation/LiveLinkHub/Source/LiveLinkHub/Private/Clients/LiveLinkHubProvider.h
+++ b/Engine/Plugins/Animation/LiveLinkHub/Source/LiveLinkHub/Private/Clients/LiveLinkHubProvider.h
@@ -10,6 +10,7 @@
 #include "LiveLinkHubMessages.h"
 #include "LiveLinkProviderImpl.h"
 #include "Misc/TransactionallySafeRWLock.h"
+#include "Misc/ScopeRWLock.h"
 #include "Templates/Function.h"
 
 
@@ -145,7 +146,7 @@ private:
 		TArray<FMessageAddress> ValidClientAddresses;
 
 		{
-			FReadScopeLock Lock(AddressBookLock);
+			UE::TReadScopeLock Lock(AddressBookLock);
 			Algo::TransformIf(AddressBook, ValidClientAddresses,
 				[this](const FLiveLinkHubAddressBookEntry& Entry) { return Entry.bConnected && ShouldTransmitToClient_AnyThread(Entry.ClientId); },
 				[](const FLiveLinkHubAddressBookEntry& Entry) { return Entry.ControlEndpoint; });
