LiveLinkHub - Increment take number when stopping a recording.

This option can be disabled in LiveLinkHubSettings - IncrementTakeOnRecordingEnd

diff --git a/Engine/Plugins/Animation/LiveLinkHub/Source/LiveLinkHub/Private/Recording/LiveLinkHubRecordingController.cpp b/Engine/Plugins/Animation/LiveLinkHub/Source/LiveLinkHub/Private/Recording/LiveLinkHubRecordingController.cpp
--- a/Engine/Plugins/Animation/LiveLinkHub/Source/LiveLinkHub/Private/Recording/LiveLinkHubRecordingController.cpp
+++ b/Engine/Plugins/Animation/LiveLinkHub/Source/LiveLinkHub/Private/Recording/LiveLinkHubRecordingController.cpp
@@ -39,6 +39,12 @@ void FLiveLinkHubRecordingController::StopRecording()
 	ILiveLinkRecordingSessionInfo::Get().OnRecordingStopped().Broadcast();
 
 	RecorderImplementation->StopRecording();
+
+	if (GetDefault<ULiveLinkHubSettings>()->bIncrementTakeOnRecordingEnd)
+	{
+		ILiveLinkRecordingSessionInfo& SessionInfo = ILiveLinkRecordingSessionInfo::Get();
+		SessionInfo.SetTakeNumber(SessionInfo.GetTakeNumber() + 1);
+	}
 }
 
 bool FLiveLinkHubRecordingController::IsRecording() const
diff --git a/Engine/Plugins/Animation/LiveLinkHub/Source/LiveLinkHub/Private/Settings/LiveLinkHubSettings.h b/Engine/Plugins/Animation/LiveLinkHub/Source/LiveLinkHub/Private/Settings/LiveLinkHubSettings.h
--- a/Engine/Plugins/Animation/LiveLinkHub/Source/LiveLinkHub/Private/Settings/LiveLinkHubSettings.h
+++ b/Engine/Plugins/Animation/LiveLinkHub/Source/LiveLinkHub/Private/Settings/LiveLinkHubSettings.h
@@ -97,6 +97,10 @@ public:
 	UPROPERTY(config, EditAnywhere, Category = "LiveLinkHub")
 	bool bPromptSaveAsOnRecord = false;
 
+	/** Whether to increment take number when stopping a recording. */
+	UPROPERTY(config, EditAnywhere, Category = "Templates")
+	bool bIncrementTakeOnRecordingEnd = true;
+
 	/** Whether to enable auto-save of Live Link Hub configs. */
 	UPROPERTY(config, EditAnywhere, Category = "Autosave")
 	bool bEnableAutosave = true;
