LiveLinkHub - Add an option for playing recordings at a different rate

#rnx

diff --git a/Engine/Plugins/Animation/LiveLinkHub/Source/LiveLinkHub/Private/Recording/LiveLinkHubPlaybackController.cpp b/Engine/Plugins/Animation/LiveLinkHub/Source/LiveLinkHub/Private/Recording/LiveLinkHubPlaybackController.cpp
--- a/Engine/Plugins/Animation/LiveLinkHub/Source/LiveLinkHub/Private/Recording/LiveLinkHubPlaybackController.cpp
+++ b/Engine/Plugins/Animation/LiveLinkHub/Source/LiveLinkHub/Private/Recording/LiveLinkHubPlaybackController.cpp
@@ -170,6 +170,10 @@ TSharedRef<SWidget> FLiveLinkHubPlaybackController::GetPlaybackWidget()
 
 			GoToTime(CurrentTimeCopy);
 		})
+		.OnSetPlayRate_Lambda([this](float InPlayRate)
+		{
+			SetPlayRate(InPlayRate);
+		})
 		.OnExitPlayback_Lambda([this]()
 		{
 			EjectAndUnload([]() {}, nullptr, bClosePlaybackTabOnEject);
@@ -640,6 +644,7 @@ uint32 FLiveLinkHubPlaybackController::Run()
 		bIsPlaybackWaiting = true;
 		PlaybackEvent->Wait();
 		bIsPlaybackWaiting = false;
+
 		while (bIsPlaying)
 		{
 			if (IsEngineExitRequested())
@@ -659,8 +664,8 @@ uint32 FLiveLinkHubPlaybackController::Run()
 
 				auto SetPlayhead = [&]()
 				{
-					const double Delta = FPlatformTime::Seconds() - PlaybackStartTime;
-					double Position = bIsReverse ? StartTimestamp - Delta : StartTimestamp + Delta;
+					const double TimeDilatedDelta = (FPlatformTime::Seconds() - PlaybackStartTime) * PlayRate;
+					double Position = bIsReverse ? StartTimestamp - TimeDilatedDelta : StartTimestamp + TimeDilatedDelta;
 					Position = FMath::Clamp(Position, GetSelectionStartTime().AsSeconds(), GetSelectionEndTime().AsSeconds());
 					Playhead->SetValue(FQualifiedFrameTime(FFrameTime::FromDecimal(Position * GetFrameRate().Numerator), GetFrameRate()));
 				};
diff --git a/Engine/Plugins/Animation/LiveLinkHub/Source/LiveLinkHub/Private/Recording/LiveLinkHubPlaybackController.h b/Engine/Plugins/Animation/LiveLinkHub/Source/LiveLinkHub/Private/Recording/LiveLinkHubPlaybackController.h
--- a/Engine/Plugins/Animation/LiveLinkHub/Source/LiveLinkHub/Private/Recording/LiveLinkHubPlaybackController.h
+++ b/Engine/Plugins/Animation/LiveLinkHub/Source/LiveLinkHub/Private/Recording/LiveLinkHubPlaybackController.h
@@ -153,6 +153,18 @@ public:
 		bLoopPlayback.store(bInLoop);
 	}
 
+	/** Sets the play rate (time dilation multiplier) of playback. */
+	void SetPlayRate(float InPlayRate)
+	{
+		PlayRate = InPlayRate;
+		// Reset the PlaybackStartTime and playhead since we're modifying the play rate.
+		GoToTime(GetCurrentTime());
+		if (!bIsPaused)
+		{
+			BeginPlayback(bIsReverse);
+		}
+	}
+
 	/** Retrieve the current recording's name. */
 	FString GetRecordingName() const;
 	
@@ -276,6 +288,9 @@ private:
 	ILiveLinkClient* Client = nullptr;
 	/** Time that the playback started. */
 	double PlaybackStartTime = 0.0;
+	double TimeDilationAdjustedCurrentTime = 0.0;
+	std::atomic<float> PlayRate = 1;
+
 	/** Playhead for the current playback. */
 	TSharedPtr<FLiveLinkHubAtomicQualifiedFrameTime, ESPMode::ThreadSafe> Playhead;
 
diff --git a/Engine/Plugins/Animation/LiveLinkHub/Source/LiveLinkHub/Private/UI/Widgets/SLiveLinkHubPlaybackWidget.cpp b/Engine/Plugins/Animation/LiveLinkHub/Source/LiveLinkHub/Private/UI/Widgets/SLiveLinkHubPlaybackWidget.cpp
--- a/Engine/Plugins/Animation/LiveLinkHub/Source/LiveLinkHub/Private/UI/Widgets/SLiveLinkHubPlaybackWidget.cpp
+++ b/Engine/Plugins/Animation/LiveLinkHub/Source/LiveLinkHub/Private/UI/Widgets/SLiveLinkHubPlaybackWidget.cpp
@@ -28,6 +28,8 @@ void SLiveLinkHubPlaybackWidget::Construct(const FArguments& InArgs)
 	OnPreviousFrameDelegate = InArgs._OnPreviousFrame;
 	OnNextFrameDelegate = InArgs._OnNextFrame;
 
+	OnSetPlayRateDelegate = InArgs._OnSetPlayRate;
+
 	OnExitPlaybackDelegate = InArgs._OnExitPlayback;
 	OnSelectRecordingDelegate = InArgs._OnSelectRecording;
 
@@ -232,8 +234,25 @@ void SLiveLinkHubPlaybackWidget::Construct(const FArguments& InArgs)
 						.Image(this, &SLiveLinkHubPlaybackWidget::GetLoopIcon)
 					]
 				]
+				+SHorizontalBox::Slot()
+				.AutoWidth()
+				[
+					SNew(SComboButton)
+					.ContentPadding(FMargin(2.0))
+					.IsEnabled(IsPlaybackEnabled)
+					.HasDownArrow(true)
+					.OnGetMenuContent(this, &SLiveLinkHubPlaybackWidget::MakePlayRateMenuContent)
+					.ForegroundColor(FSlateColor::UseStyle())
+					.ToolTipText(LOCTEXT("PlaybackRate_Tooltip", "Change Play Rate."))
+					.ComboButtonStyle(&FAppStyle::Get().GetWidgetStyle<FComboButtonStyle>("SimpleComboButton"))
+					.ButtonContent()
+					[
+						SNew(STextBlock)
+						.Text_Lambda([this]() { return FText::Format(INVTEXT("{0}x"), CachedPlayRate); })
+					]
+				]
 				+ SHorizontalBox::Slot()
-				.Padding(FMargin(12.0, 0.0, 0.0, 0.0))
+				.Padding(FMargin(4.0, 0.0, 0.0, 0.0))
 				.AutoWidth()
 				[
 					SNew(SWidgetSwitcher)
@@ -624,6 +643,12 @@ bool SLiveLinkHubPlaybackWidget::CompareDisplayFormat(EFrameNumberDisplayFormats
 	return DisplayFormat == InDisplayFormat;
 }
 
+void SLiveLinkHubPlaybackWidget::SetPlayRate(float InPlayRate)
+{
+	CachedPlayRate = FMath::Clamp(InPlayRate, 0.01, 5);
+	OnSetPlayRateDelegate.ExecuteIfBound(InPlayRate);
+}
+
 FText SLiveLinkHubPlaybackWidget::GetDisplayFormatAsText() const
 {
 	return DisplayFormat == EFrameNumberDisplayFormats::Frames ? LOCTEXT("DisplayFormat_TimeFrames", "Frames") :
@@ -719,6 +744,38 @@ TSharedRef<SWidget> SLiveLinkHubPlaybackWidget::MakePlaybackSettingsDropdown()
 	return MenuBuilder.MakeWidget();
 }
 
+TSharedRef<SWidget> SLiveLinkHubPlaybackWidget::MakePlayRateMenuContent()
+{
+	constexpr bool bCloseAfterSelection = true;
+	FMenuBuilder MenuBuilder(bCloseAfterSelection, nullptr);
+
+	const FCanExecuteAction AlwaysExecute = FCanExecuteAction::CreateLambda([]{ return true; });
+
+	MenuBuilder.BeginSection("PlayRateSection", LOCTEXT("PlayRate", "Play Rate"));
+
+	TArray<float, TInlineAllocator<6>> Values = { 0.25f, 0.5f, 0.75f, 1.0f, 1.25f, 1.5f, 2.0f };
+
+	for (float Value : Values)
+	{
+		MenuBuilder.AddMenuEntry(
+			FText::Format(INVTEXT("{0}x"), Value),
+			FText::Format(LOCTEXT("PlayRateToolTip", "Set Play Rate to {0}x"), Value),
+			FSlateIcon(),
+			FUIAction(
+				FExecuteAction::CreateSP(this, &SLiveLinkHubPlaybackWidget::SetPlayRate, Value),
+				AlwaysExecute,
+				FIsActionChecked::CreateLambda([this, Value]() -> bool { return Value == CachedPlayRate; })
+			),
+			NAME_None,
+			EUserInterfaceActionType::RadioButton
+		);
+	}
+	
+	MenuBuilder.EndSection();
+	
+	return MenuBuilder.MakeWidget();
+}
+
 TSharedRef<SWidget> SLiveLinkHubPlaybackWidget::MakeUpgradingWidget()
 {
 	return
diff --git a/Engine/Plugins/Animation/LiveLinkHub/Source/LiveLinkHub/Private/UI/Widgets/SLiveLinkHubPlaybackWidget.h b/Engine/Plugins/Animation/LiveLinkHub/Source/LiveLinkHub/Private/UI/Widgets/SLiveLinkHubPlaybackWidget.h
--- a/Engine/Plugins/Animation/LiveLinkHub/Source/LiveLinkHub/Private/UI/Widgets/SLiveLinkHubPlaybackWidget.h
+++ b/Engine/Plugins/Animation/LiveLinkHub/Source/LiveLinkHub/Private/UI/Widgets/SLiveLinkHubPlaybackWidget.h
@@ -27,6 +27,7 @@ public:
 	DECLARE_DELEGATE_RetVal(TRange<double>, FOnGetDoubleRange)
 	DECLARE_DELEGATE_RetVal(UE::LiveLinkHub::RangeHelpers::Private::TRangeArray<int32>, FOnGetIntArrayRange)
 	DECLARE_DELEGATE_RetVal(float, FOnGetFloat)
+	DECLARE_DELEGATE_OneParam(FOnSetFloat, float)
 	
 	DECLARE_DELEGATE(FOnButtonPressed);
 	
@@ -43,6 +44,8 @@ public:
 	SLATE_EVENT(FOnButtonPressed, OnPreviousFrame)
 	/** When next frame is pressed. */
 	SLATE_EVENT(FOnButtonPressed, OnNextFrame)
+	/** When next frame is pressed. */
+	SLATE_EVENT(FOnSetFloat, OnSetPlayRate)
 	/** When the eject button is pressed. */
 	SLATE_EVENT(FOnButtonPressed, OnExitPlayback)
 	/** When the select recording button is pressed. */
@@ -173,6 +176,9 @@ private:
 	/** Validates the current display format is set to the given format. */
 	bool CompareDisplayFormat(EFrameNumberDisplayFormats InDisplayFormat) const;
 
+	/** Set the play rate (time dilation multiplier). */
+	void SetPlayRate(float InPlayRate);
+
 	/** Retrieve the text version of the display format. */
 	FText GetDisplayFormatAsText() const;
 	
@@ -194,6 +200,9 @@ private:
 	/** Create the dropdown box for playback settings. */
 	TSharedRef<SWidget> MakePlaybackSettingsDropdown();
 
+	/** Make the menu content that allows picking a play rate. */
+	TSharedRef<SWidget> MakePlayRateMenuContent();
+
 	/** The widget used for displaying upgrade notifications. */
 	TSharedRef<SWidget> MakeUpgradingWidget();
 
@@ -227,6 +236,9 @@ private:
 	/** Delegate for going to the next frame. */
 	FOnButtonPressed OnNextFrameDelegate;
 
+	/** Delegate for setting the play rate. */
+	FOnSetFloat OnSetPlayRateDelegate;
+
 	/** Delegate for exiting the current recording. */
 	FOnButtonPressed OnExitPlaybackDelegate;
 
@@ -318,6 +330,9 @@ private:
 	/** Counter used to advance the recording while holding down left or right key. */
 	uint8 KeyDownScrubCounter = 0;
 
+	/** Cached value of the recording playback rate. */
+	float CachedPlayRate = 1.0f;
+
 	/** Number of frames to wait between advancing the frame while a key is held down.*/
 	static constexpr uint8 FramesBetweenKeyDownScrub = 3;
 
