LiveLinkHub - Add an option to allow receiving data from Unreal

Move the CanReceiveFrom and CanTransmitTo logic from a static map in LiveLinkProvider.cpp to LiveLinkhubMessagingSettings.

diff --git a/Engine/Plugins/Animation/LiveLinkHub/Source/LiveLinkHub/Private/Clients/LiveLinkHubProvider.cpp b/Engine/Plugins/Animation/LiveLinkHub/Source/LiveLinkHub/Private/Clients/LiveLinkHubProvider.cpp
--- a/Engine/Plugins/Animation/LiveLinkHub/Source/LiveLinkHub/Private/Clients/LiveLinkHubProvider.cpp
+++ b/Engine/Plugins/Animation/LiveLinkHub/Source/LiveLinkHub/Private/Clients/LiveLinkHubProvider.cpp
@@ -16,6 +16,7 @@
 #include "INetworkMessagingExtension.h"
 #include "LiveLinkHubLog.h"
 #include "LiveLinkHubMessages.h"
+#include "LiveLinkHubMessagingSettings.h"
 #include "LiveLinkSettings.h"
 #include "MessageEndpointBuilder.h"
 #include "Misc/ScopeLock.h"
@@ -377,7 +378,7 @@ void FLiveLinkHubProvider::HandleHubConnectMessage(const FLiveLinkHubConnectMess
 	UE_LOG(LogLiveLinkHub, Verbose, TEXT("Provider: Received connect message from %s"), *Context->GetSender().ToString());
 
 	const ELiveLinkTopologyMode Mode = FLiveLinkHub::Get()->GetTopologyMode();
-	if (!UE::LiveLink::Messaging::CanTransmitTo(Mode, Message.ClientInfo.TopologyMode))
+	if (!GetDefault<ULiveLinkHubMessagingSettings>()->CanTransmitTo(Mode, Message.ClientInfo.TopologyMode))
 	{
 		UE_LOG(LogLiveLinkHub, Verbose, TEXT("Provider: Denying connection from %s since its mode is incompatible with this provider's."), *Context->GetSender().ToString());
 		FLiveLinkHubDisconnectMessage DisconnectMessage{ GetProviderName(), GetMachineName() };
@@ -648,7 +649,7 @@ void FLiveLinkHubProvider::ConnectToAllDiscoveredClients()
 	{
 		UE::TReadScopeLock Lock(AddressBookLock);
 		Algo::TransformIf(AddressBook, CompatibleClients,
-			[](const FLiveLinkHubAddressBookEntry& Entry) { return UE::LiveLink::Messaging::CanTransmitTo(FLiveLinkHub::Get()->GetTopologyMode(), Entry.TopologyMode); },
+			[](const FLiveLinkHubAddressBookEntry& Entry) { return GetDefault<ULiveLinkHubMessagingSettings>()->CanTransmitTo(FLiveLinkHub::Get()->GetTopologyMode(), Entry.TopologyMode); },
 			[](const FLiveLinkHubAddressBookEntry& Entry) { return Entry.ControlEndpoint; });
 	}
 
@@ -745,7 +746,7 @@ void FLiveLinkHubProvider::ValidateHubConnections()
 						It->LastDiscoveryRequestTimestamp = FPlatformTime::Seconds();
 
 						if (ShouldAutoConnectTo(It->Hostname) 
-							&& (UE::LiveLink::Messaging::CanTransmitTo(FLiveLinkHub::Get()->GetTopologyMode(), It->TopologyMode))) // Previous versions of UE/LLH relied on the client side to check for Hub/Spoke compatibility.
+							&& (GetDefault<ULiveLinkHubMessagingSettings>()->CanTransmitTo(FLiveLinkHub::Get()->GetTopologyMode(), It->TopologyMode))) // Previous versions of UE/LLH relied on the client side to check for Hub/Spoke compatibility.
 						{
 							It->bSentDiscoveryRequest = true;
 
@@ -825,7 +826,7 @@ TArray<FLiveLinkHubClientId> FLiveLinkHubProvider::GetDiscoveredClients() const
 
 			if (Entry.DiscoveryProtocolVersion > 1)
 			{
-				return UE::LiveLink::Messaging::CanTransmitTo(FLiveLinkHub::Get()->GetTopologyMode(), Entry.TopologyMode);
+				return GetDefault<ULiveLinkHubMessagingSettings>()->CanTransmitTo(FLiveLinkHub::Get()->GetTopologyMode(), Entry.TopologyMode);
 			}
 			else
 			{
diff --git a/Engine/Plugins/Animation/LiveLinkHub/Source/LiveLinkHubMessaging/LiveLinkHubMessaging.build.cs b/Engine/Plugins/Animation/LiveLinkHub/Source/LiveLinkHubMessaging/LiveLinkHubMessaging.build.cs
--- a/Engine/Plugins/Animation/LiveLinkHub/Source/LiveLinkHubMessaging/LiveLinkHubMessaging.build.cs
+++ b/Engine/Plugins/Animation/LiveLinkHub/Source/LiveLinkHubMessaging/LiveLinkHubMessaging.build.cs
@@ -23,6 +23,7 @@ namespace UnrealBuildTool.Rules
 				"LiveLinkInterface",
 				"Messaging",
 				"SlateCore",
+				"Settings",
 			});
 
 			if (Target.bCompileAgainstEditor)
diff --git a/Engine/Plugins/Animation/LiveLinkHub/Source/LiveLinkHubMessaging/Private/LiveLinkHubConnectionManager.h b/Engine/Plugins/Animation/LiveLinkHub/Source/LiveLinkHubMessaging/Private/LiveLinkHubConnectionManager.h
--- a/Engine/Plugins/Animation/LiveLinkHub/Source/LiveLinkHubMessaging/Private/LiveLinkHubConnectionManager.h
+++ b/Engine/Plugins/Animation/LiveLinkHub/Source/LiveLinkHubMessaging/Private/LiveLinkHubConnectionManager.h
@@ -12,6 +12,7 @@
 #include "LiveLinkHubMessages.h"
 #include "LiveLinkHubMessageBusSource.h"
 #include "LiveLinkHubMessagingModule.h"
+#include "LiveLinkHubMessagingSettings.h"
 #include "LiveLinkMessageBusDiscoveryManager.h"
 #include "LiveLinkMessageBusFinder.h"
 #include "LiveLinkSettings.h"
@@ -113,7 +114,7 @@ namespace LiveLinkHubConnectionManager
 		// Topology Mode
 		ELiveLinkTopologyMode IncomingMode = LiveLinkHubConnectionManager::GetPollResultTopologyMode(InPollResult);
 
-		const bool bCompatibleModeResult = UE::LiveLink::Messaging::CanReceiveFrom(InHostMode, IncomingMode);
+		const bool bCompatibleModeResult = GetDefault<ULiveLinkHubMessagingSettings>()->CanReceiveFrom(InHostMode, IncomingMode);
 
 		const FText ModeName = UEnum::GetDisplayValueAsText(InHostMode);
 		const FText IncomingModeName = UEnum::GetDisplayValueAsText(IncomingMode);
diff --git a/Engine/Plugins/Animation/LiveLinkHub/Source/LiveLinkHubMessaging/Private/LiveLinkHubControlChannel.h b/Engine/Plugins/Animation/LiveLinkHub/Source/LiveLinkHubMessaging/Private/LiveLinkHubControlChannel.h
--- a/Engine/Plugins/Animation/LiveLinkHub/Source/LiveLinkHubMessaging/Private/LiveLinkHubControlChannel.h
+++ b/Engine/Plugins/Animation/LiveLinkHub/Source/LiveLinkHubMessaging/Private/LiveLinkHubControlChannel.h
@@ -271,7 +271,7 @@ private:
 		FLiveLinkHubMessagingModule& Module = FModuleManager::Get().GetModuleChecked<FLiveLinkHubMessagingModule>("LiveLinkHubMessaging");
 		ELiveLinkTopologyMode Mode = Module.GetHostTopologyMode();
 
-		bool bShouldConnect = UE::LiveLink::Messaging::CanReceiveFrom(Mode, Message.Mode)
+		bool bShouldConnect = GetDefault<ULiveLinkHubMessagingSettings>()->CanReceiveFrom(Mode, Message.Mode)
 			&& LiveLinkHubConnectionManager::CanConnectTo(Message.MachineName, *Context, Module.GetInstanceId());
 		if (bShouldConnect)
 		{
diff --git a/Engine/Plugins/Animation/LiveLinkHub/Source/LiveLinkHubMessaging/Private/LiveLinkHubMessagingModule.cpp b/Engine/Plugins/Animation/LiveLinkHub/Source/LiveLinkHubMessaging/Private/LiveLinkHubMessagingModule.cpp
--- a/Engine/Plugins/Animation/LiveLinkHub/Source/LiveLinkHubMessaging/Private/LiveLinkHubMessagingModule.cpp
+++ b/Engine/Plugins/Animation/LiveLinkHub/Source/LiveLinkHubMessaging/Private/LiveLinkHubMessagingModule.cpp
@@ -4,6 +4,7 @@
 
 #include "CoreMinimal.h"
 #include "IMessageInterceptor.h"
+#include "ISettingsModule.h"
 #include "LiveLinkHubConnectionManager.h"
 #include "LiveLinkHubControlChannel.h"
 #include "LiveLinkHubMessageBusSourceFactory.h"
@@ -19,6 +20,8 @@
 
 DEFINE_LOG_CATEGORY(LogLiveLinkHubMessaging);
 
+#define LOCTEXT_NAMESPACE "LiveLinkHubMessaging"
+
 namespace Internal
 {
 	// Source: https://en.cppreference.com/w/cpp/utility/variant/visit
@@ -51,6 +54,8 @@ void FLiveLinkHubMessagingModule::StartupModule()
 
 #endif
 
+	RegisterSettings();
+
 	SourceFilterDelegate = ILiveLinkModule::Get().RegisterMessageBusSourceFilter(FOnLiveLinkShouldDisplaySource::CreateRaw(this, &FLiveLinkHubMessagingModule::OnFilterMessageBusSource));
 }
 
@@ -61,6 +66,8 @@ void FLiveLinkHubMessagingModule::ShutdownModule()
 		LiveLinkModule->UnregisterMessageBusSourceFilter(SourceFilterDelegate);
 	}
 
+	UnregisterSettings();
+
 #if WITH_LIVELINK_DISCOVERY_MANAGER_THREAD
 	ControlChannel.Reset();
 	ConnectionManager.Reset();
@@ -122,6 +129,33 @@ TSharedPtr<FLiveLinkHubControlChannel> FLiveLinkHubMessagingModule::GetControlCh
 	return ControlChannel;
 }
 
+void FLiveLinkHubMessagingModule::RegisterSettings()
+{
+	// Only create this section if in the hub.
+	if (GConfig->GetBoolOrDefault(TEXT("LiveLink"), TEXT("bCreateLiveLinkHubInstance"), false, GEngineIni))
+	{
+		if (ISettingsModule* SettingsModule = FModuleManager::GetModulePtr<ISettingsModule>("Settings"))
+		{
+			SettingsModule->RegisterSettings("Project", "Application", "Messaging",
+				LOCTEXT("LiveLinkSettingsName", "Messaging"),
+				LOCTEXT("LiveLinkDescription", "Configure Live Link Hub Messaging"),
+				GetMutableDefault<ULiveLinkHubMessagingSettings>()
+			);
+		}
+	}
+}
+
+void FLiveLinkHubMessagingModule::UnregisterSettings()
+{
+	if (GConfig->GetBoolOrDefault(TEXT("LiveLink"), TEXT("bCreateLiveLinkHubInstance"), false, GEngineIni))
+	{
+		if (ISettingsModule* SettingsModule = FModuleManager::GetModulePtr<ISettingsModule>("Settings"))
+		{
+			SettingsModule->UnregisterSettings("Project", "Application", "Timing & Sync");
+		}
+	}
+}
+
 FLiveLinkHubInstanceId::FLiveLinkHubInstanceId(FGuid Guid)
 {
 	Id.Set<FGuid>(Guid);
@@ -151,3 +185,5 @@ FString FLiveLinkHubInstanceId::ToString() const
 }
 
 IMPLEMENT_MODULE(FLiveLinkHubMessagingModule, LiveLinkHubMessaging);
+
+#undef LOCTEXT_NAMESPACE
\ No newline at end of file
diff --git a/Engine/Plugins/Animation/LiveLinkHub/Source/LiveLinkHubMessaging/Private/LiveLinkHubMessagingModule.h b/Engine/Plugins/Animation/LiveLinkHub/Source/LiveLinkHubMessaging/Private/LiveLinkHubMessagingModule.h
--- a/Engine/Plugins/Animation/LiveLinkHub/Source/LiveLinkHubMessaging/Private/LiveLinkHubMessagingModule.h
+++ b/Engine/Plugins/Animation/LiveLinkHub/Source/LiveLinkHubMessaging/Private/LiveLinkHubMessagingModule.h
@@ -47,6 +47,12 @@ private:
 	/** Handle a discovery request from LLH to populate its address book with the ControlEndpoint's address. */
 	void OnDiscoveryRequest(const struct FMessageAddress& RemoteAddress) const;
 
+	/** Register settings related to LiveLinkHub messaging. */
+	void RegisterSettings();
+
+	/** Unregister settings related to LiveLinkHub messaging. */
+	void UnregisterSettings();
+
 private:
 	bool bUseConnectionManager;
 
diff --git a/Engine/Plugins/Animation/LiveLinkHub/Source/LiveLinkHubMessaging/Public/LiveLinkHubMessagingSettings.h b/Engine/Plugins/Animation/LiveLinkHub/Source/LiveLinkHubMessaging/Public/LiveLinkHubMessagingSettings.h
new file mode 100644
--- /dev/null
+++ b/Engine/Plugins/Animation/LiveLinkHub/Source/LiveLinkHubMessaging/Public/LiveLinkHubMessagingSettings.h
@@ -0,0 +1,114 @@
+// Copyright Epic Games, Inc. All Rights Reserved.
+
+#pragma once
+
+#include "UObject/Object.h"
+#include "LiveLinkHubMessages.h"
+
+#include "LiveLinkHubMessagingSettings.generated.h"
+
+/**
+ * Settings for the messaging protocol of LiveLinkHub.
+ */
+UCLASS(MinimalAPI, config = EditorPerProjectUserSettings)
+class ULiveLinkHubMessagingSettings : public UObject
+{
+	GENERATED_BODY()
+
+public:
+
+	ULiveLinkHubMessagingSettings()
+	{
+		using enum ELiveLinkTopologyMode;
+
+		CanReceiveFrom_Map =
+		{
+			{ UnrealClient, Hub },
+			{ UnrealClient, UnrealClient },
+			{ UnrealClient, External },
+			{ Spoke, External },
+			{ Hub, External },
+			{ Hub, Spoke },
+		};
+
+		CanTransmitTo_Map =
+		{
+			{ Hub, UnrealClient },
+			{ UnrealClient, UnrealClient },
+			{ Spoke, Hub }
+		};
+	}
+
+	/** Whether to allow the hub to receive LiveLink data from an unreal instance. (Useful if Unreal is using Live Link Broadcast components. */
+	UPROPERTY(config, EditAnywhere, Category = "LiveLinkHub")
+	bool bAllowReceivingFromUnreal = false;
+
+public:
+	//~ UObject interface
+	virtual void PostInitProperties() override
+	{
+		using enum ELiveLinkTopologyMode;
+
+		Super::PostInitProperties();
+
+		if (bAllowReceivingFromUnreal)
+		{
+			CanReceiveFrom_Map.Add(Hub, UnrealClient);
+			CanTransmitTo_Map.Add(UnrealClient, Hub);
+		}
+	}
+
+#if WITH_EDITOR
+	void PostEditChangeProperty(FPropertyChangedEvent& PropertyChangedEvent)
+	{
+		using enum ELiveLinkTopologyMode;
+
+		UObject::PostEditChangeProperty(PropertyChangedEvent);
+
+		if (PropertyChangedEvent.GetPropertyName() == GET_MEMBER_NAME_CHECKED(ULiveLinkHubMessagingSettings, bAllowReceivingFromUnreal))
+		{
+			if (bAllowReceivingFromUnreal)
+			{
+				CanReceiveFrom_Map.Add(Hub, UnrealClient);
+				CanTransmitTo_Map.Add(UnrealClient, Hub);
+			}
+			else
+			{
+				CanReceiveFrom_Map.Remove(Hub, UnrealClient);
+				CanTransmitTo_Map.Remove(UnrealClient, Hub);
+			}
+		}
+	}
+#endif
+	//~ End UObject interface
+
+	/** Whether an instance with the LHS mode can act as a Source for the RHS mode instance. */
+	bool CanReceiveFrom(ELiveLinkTopologyMode LHSMode, ELiveLinkTopologyMode RHSMode) const
+	{
+		TArray<ELiveLinkTopologyMode> AllowedConnections;
+		CanReceiveFrom_Map.MultiFind(LHSMode, AllowedConnections);
+
+		return AllowedConnections.Contains(RHSMode);
+	}
+
+	/** Whether an instance with the LHS mode can act as a Provider for the RHS mode instance. */
+	bool CanTransmitTo(ELiveLinkTopologyMode LHSMode, ELiveLinkTopologyMode RHSMode) const
+	{
+		TArray<ELiveLinkTopologyMode> AllowedConnections;
+		CanTransmitTo_Map.MultiFind(LHSMode, AllowedConnections);
+		return AllowedConnections.Contains(RHSMode);
+	}
+
+private:
+	/**
+	 * Map of possible connections.
+	 * Left side can receive a connection from right side.
+	 */
+	TMultiMap<ELiveLinkTopologyMode, ELiveLinkTopologyMode> CanReceiveFrom_Map;
+
+	/**
+		* Map of possible connections.
+		* Left side can transmit data to a client from right side.
+		*/
+	TMultiMap<ELiveLinkTopologyMode, ELiveLinkTopologyMode> CanTransmitTo_Map;
+};
diff --git a/Engine/Source/Programs/LiveLinkHub/Config/DefaultEngine.ini b/Engine/Source/Programs/LiveLinkHub/Config/DefaultEngine.ini
--- a/Engine/Source/Programs/LiveLinkHub/Config/DefaultEngine.ini
+++ b/Engine/Source/Programs/LiveLinkHub/Config/DefaultEngine.ini
@@ -28,6 +28,7 @@ AutoConnectClients=LocalOnly
 +ProjectSettingsToDisplay=Live Link
 +ProjectSettingsToDisplay=Live Link Hub
 +ProjectSettingsToDisplay=Timing & Sync
++ProjectSettingsToDisplay=Messaging
 +ProjectSettingsToDisplay=UDPMessaging
 bShowFrameRate=True
 TargetFrameRate=60.000000
diff --git a/Engine/Source/Runtime/LiveLinkMessageBusFramework/Private/LiveLinkProvider.cpp b/Engine/Source/Runtime/LiveLinkMessageBusFramework/Private/LiveLinkProvider.cpp
--- a/Engine/Source/Runtime/LiveLinkMessageBusFramework/Private/LiveLinkProvider.cpp
+++ b/Engine/Source/Runtime/LiveLinkMessageBusFramework/Private/LiveLinkProvider.cpp
@@ -36,52 +36,6 @@ static TAutoConsoleVariable<int32> CVarLiveLinkCompressAnimData(
 	TEXT("Whether to compress animation data that goes through the provider."),
 	ECVF_RenderThreadSafe);
 
-
-namespace UE::LiveLink::Messaging
-{
-	using enum ELiveLinkTopologyMode;
-
-	/**
-	 * Map of possible connections.
-	 * Left side can receive a connection from right side.
-	 */
-	static const TMultiMap<ELiveLinkTopologyMode, ELiveLinkTopologyMode> CanReceiveFrom_Map =
-	{
-		{ UnrealClient, Hub },
-		{ UnrealClient, UnrealClient },
-		{ UnrealClient, External },
-		{ Spoke, External },
-		{ Hub, External },
-		{ Hub, Spoke },
-	};
-
-	/**
-	 * Map of possible connections.
-	 * Left side can transmit data to a client from right side.
-	 */
-	static const TMultiMap<ELiveLinkTopologyMode, ELiveLinkTopologyMode> CanTransmitTo_Map =
-	{
-		{ Hub, UnrealClient },
-		{ UnrealClient, UnrealClient },
-		{ Spoke, Hub }
-	};
-
-	bool CanReceiveFrom(ELiveLinkTopologyMode LHSMode, ELiveLinkTopologyMode RHSMode)
-	{
-		TArray<ELiveLinkTopologyMode> AllowedConnections;
-		CanReceiveFrom_Map.MultiFind(LHSMode, AllowedConnections);
-		return AllowedConnections.Contains(RHSMode);
-	}
-
-	bool CanTransmitTo(ELiveLinkTopologyMode LHSMode, ELiveLinkTopologyMode RHSMode)
-	{
-		TArray<ELiveLinkTopologyMode> AllowedConnections;
-		CanTransmitTo_Map.MultiFind(LHSMode, AllowedConnections);
-		return AllowedConnections.Contains(RHSMode);
-	}
-}
-
-
 // Address that we have had a connection request from
 struct FTrackedAddress
 {
