Live Link Hub: Add "auxiliary channel" endpoint negotiation.

It's now possible to register a channel type handler with ILiveLinkHubMessagingModule, and initiate a channel request to a known client using ILiveLinkHubClientsModel. The auxiliary channel will be negotiated over the existing control channel while keeping Live Link Hub's internal endpoints encapsulated.

Other changes:
- ILiveLinkRecordingSessionInfo is now ILiveLinkRecordingSession, and the interface has been expanded to include CanRecord/StartRecording/StopRecording methods. The backward compatibility header + typedef will be deprecated in a subsequent changelist.
- ILiveLinkHubClientsModel is now public and a modular feature.
- Add property type customization for FLiveLinkHubClientId, with a dropdown selector for connected clients.

diff --git a/Engine/Plugins/Animation/LiveLinkHub/Source/LiveLinkHub/Private/Clients/LiveLinkHubClientsModel.h b/Engine/Plugins/Animation/LiveLinkHub/Source/LiveLinkHub/Private/Clients/LiveLinkHubClientsModel.h
deleted file mode 100644
diff --git a/Engine/Plugins/Animation/LiveLinkHub/Source/LiveLinkHub/Private/Clients/LiveLinkHubProvider.cpp b/Engine/Plugins/Animation/LiveLinkHub/Source/LiveLinkHub/Private/Clients/LiveLinkHubProvider.cpp
--- a/Engine/Plugins/Animation/LiveLinkHub/Source/LiveLinkHub/Private/Clients/LiveLinkHubProvider.cpp
+++ b/Engine/Plugins/Animation/LiveLinkHub/Source/LiveLinkHub/Private/Clients/LiveLinkHubProvider.cpp
@@ -6,7 +6,6 @@
 #include "Algo/RemoveIf.h"
 #include "Algo/Transform.h"
 #include "Async/Async.h"
-#include "Clients/LiveLinkHubClientsModel.h"
 #include "Clients/LiveLinkHubUEClientInfo.h"
 #include "Containers/ObservableArray.h"
 #include "CoreMinimal.h"
@@ -18,6 +17,7 @@
 #include "LiveLinkHubMessages.h"
 #include "LiveLinkHubMessagingSettings.h"
 #include "LiveLinkSettings.h"
+#include "Logging/StructuredLog.h"
 #include "MessageEndpointBuilder.h"
 #include "Misc/ScopeLock.h"
 #include "Session/LiveLinkHubSession.h"
@@ -96,6 +96,8 @@ FLiveLinkHubProvider::FLiveLinkHubProvider(const TSharedRef<ILiveLinkHubSessionM
 		const double ValidateConnectionsRate = GetDefault<ULiveLinkSettings>()->MessageBusPingRequestFrequency;
         GEditor->GetTimerManager()->SetTimer(ValidateConnectionsTimer, FTimerDelegate::CreateRaw(this, &FLiveLinkHubProvider::ValidateHubConnections), ValidateConnectionsRate, true);
 	});
+
+	IModularFeatures::Get().RegisterModularFeature(ILiveLinkHubClientsModel::GetModularFeatureName(), this);
 }
 
 FLiveLinkHubProvider::~FLiveLinkHubProvider()
@@ -104,6 +106,8 @@ FLiveLinkHubProvider::~FLiveLinkHubProvider()
 	{
 		GEditor->GetTimerManager()->ClearTimer(ValidateConnectionsTimer);
 	}
+
+	IModularFeatures::Get().UnregisterModularFeature(ILiveLinkHubClientsModel::GetModularFeatureName(), this);
 }
 
 bool FLiveLinkHubProvider::ShouldTransmitToSubject_AnyThread(FName SubjectName, FMessageAddress Address) const
@@ -1011,6 +1015,43 @@ void FLiveLinkHubProvider::SetSubjectEnabled(FLiveLinkHubClientId Client, FName
 	}
 }
 
+void FLiveLinkHubProvider::RequestAuxiliaryChannel(
+	FLiveLinkHubClientId InClientId,
+	FMessageEndpoint& InAuxEndpoint,
+	UScriptStruct* InRequestTypeInfo,
+	FLiveLinkHubAuxChannelRequestMessage& InRequest
+)
+{
+	if (!InClientId.IsValid())
+	{
+		UE_LOGFMT(LogLiveLinkHub, Warning, "RequestAuxiliaryChannel: Invalid client ID.");
+		return;
+	}
+
+	FMessageAddress RelevantControlEndpoint;
+
+	{
+		UE::TReadScopeLock Lock(AddressBookLock);
+		for (const FLiveLinkHubAddressBookEntry& Entry : AddressBook)
+		{
+			if (Entry.bConnected && Entry.ClientId == InClientId)
+			{
+				RelevantControlEndpoint = Entry.ControlEndpoint;
+				break;
+			}
+		}
+	}
+
+	if (RelevantControlEndpoint.IsValid())
+	{
+		InAuxEndpoint.Send(&InRequest, InRequestTypeInfo, EMessageFlags::Reliable, {}, nullptr, { RelevantControlEndpoint }, FTimespan::Zero(), FDateTime::MaxValue());
+	}
+	else
+	{
+		UE_LOGFMT(LogLiveLinkHub, Warning, "Could not find address for client {ClientId}.", InClientId.ToString());
+	}
+}
+
 FLiveLinkHubClientId FLiveLinkHubProvider::UpdateAddressBookEntry(const FMessageAddress& ControlEndpointAddress, int32 DiscoveryProtocolVersion, ELiveLinkTopologyMode TopologyMode, const FString& Hostname)
 {
 	TOptional<FLiveLinkHubAddressBookEntry> ReadOnlyEntry;
diff --git a/Engine/Plugins/Animation/LiveLinkHub/Source/LiveLinkHub/Private/Clients/LiveLinkHubProvider.h b/Engine/Plugins/Animation/LiveLinkHub/Source/LiveLinkHub/Private/Clients/LiveLinkHubProvider.h
--- a/Engine/Plugins/Animation/LiveLinkHub/Source/LiveLinkHub/Private/Clients/LiveLinkHubProvider.h
+++ b/Engine/Plugins/Animation/LiveLinkHub/Source/LiveLinkHub/Private/Clients/LiveLinkHubProvider.h
@@ -3,9 +3,9 @@
 #pragma once
 
 #include "Clients/GameThreadMessageHandler.h"
-#include "Clients/LiveLinkHubClientsModel.h"
 #include "Clients/LiveLinkHubUEClientInfo.h"
 #include "Engine/TimerHandle.h"
+#include "ILiveLinkHubClientsModel.h"
 #include "IMessageContext.h"
 #include "LiveLinkHubMessages.h"
 #include "LiveLinkProviderImpl.h"
@@ -79,6 +79,7 @@ public:
 	//~ Begin LiveLinkProvider interface
 	virtual bool ShouldTransmitToSubject_AnyThread(FName SubjectName, FMessageAddress Address) const override;
 	virtual TOptional<FLiveLinkHubUEClientInfo> GetClientInfo(FLiveLinkHubClientId InClient) const override;
+	virtual void OnMessageBusNotification(const FMessageBusNotification& Notification) override;
 	//~ End LiveLinkProvider interface
 
 	/**
@@ -208,7 +209,7 @@ protected:
 	virtual void SetClientEnabled(FLiveLinkHubClientId Client, bool bInEnable) override;
 	virtual bool IsSubjectEnabled(FLiveLinkHubClientId Client, FName SubjectName) const override;
 	virtual void SetSubjectEnabled(FLiveLinkHubClientId Client, FName SubjectName, bool bInEnable) override;
-	virtual void OnMessageBusNotification(const FMessageBusNotification& Notification) override;
+	virtual void RequestAuxiliaryChannel(FLiveLinkHubClientId InClientId, FMessageEndpoint& InAuxEndpoint, UScriptStruct* InRequestTypeInfo, FLiveLinkHubAuxChannelRequestMessage& InRequest) override;
 	//~ End ILiveLinkHubClientsModel interface
 
 
diff --git a/Engine/Plugins/Animation/LiveLinkHub/Source/LiveLinkHub/Private/Clients/LiveLinkHubUEClientInfo.h b/Engine/Plugins/Animation/LiveLinkHub/Source/LiveLinkHub/Private/Clients/LiveLinkHubUEClientInfo.h
--- a/Engine/Plugins/Animation/LiveLinkHub/Source/LiveLinkHub/Private/Clients/LiveLinkHubUEClientInfo.h
+++ b/Engine/Plugins/Animation/LiveLinkHub/Source/LiveLinkHub/Private/Clients/LiveLinkHubUEClientInfo.h
@@ -2,55 +2,11 @@
 
 #pragma once
 
-#include "IMessageContext.h"
+#include "ILiveLinkHubClientsModel.h"
 #include "LiveLinkHubMessages.h"
 
 #include "LiveLinkHubUEClientInfo.generated.h"
 
-/**
- * Identifier for a UE client receiving data from the hub.
- */
-USTRUCT()
-struct FLiveLinkHubClientId
-{
-	GENERATED_BODY()
-
-	/** Default constructor, should only be used by the reflection code. */
-	FLiveLinkHubClientId() = default;
-
-	static FLiveLinkHubClientId NewId()
-	{
-		FLiveLinkHubClientId Id{};
-		Id.Guid = FGuid::NewGuid();
-		return Id;
-	}
-
-	bool IsValid() const
-	{
-		return Guid.IsValid();
-	}
-
-	/** Get the hash of this ID. */
-	friend uint32 GetTypeHash(FLiveLinkHubClientId Id)
-	{
-		return GetTypeHash(Id.Guid);
-	}
-
-	bool operator==(const FLiveLinkHubClientId& Other) const
-	{
-		return Guid == Other.Guid;
-	}
-
-	FString ToString() const
-	{
-		return Guid.ToString();
-	}
-
-private:
-	/** Unique identifier for this client. */
-	UPROPERTY()
-	FGuid Guid;
-};
 
 /** Wrapper around FLiveLinkClientInfoMessage that adds additional info. Used mainly to display information about a client in the UI. */
 USTRUCT()
diff --git a/Engine/Plugins/Animation/LiveLinkHub/Source/LiveLinkHub/Private/Clients/SLiveLinkHubClientsView.h b/Engine/Plugins/Animation/LiveLinkHub/Source/LiveLinkHub/Private/Clients/SLiveLinkHubClientsView.h
--- a/Engine/Plugins/Animation/LiveLinkHub/Source/LiveLinkHub/Private/Clients/SLiveLinkHubClientsView.h
+++ b/Engine/Plugins/Animation/LiveLinkHub/Source/LiveLinkHub/Private/Clients/SLiveLinkHubClientsView.h
@@ -13,11 +13,11 @@
 #include "Framework/Application/SlateApplication.h"
 #include "Framework/MultiBox/MultiBoxBuilder.h"
 #include "Framework/Commands/UIAction.h"
-#include "Misc/Guid.h"
-#include "Misc/TextFilter.h"
+#include "ILiveLinkHubClientsModel.h"
 #include "LiveLinkClient.h"
-#include "LiveLinkHubClientsModel.h"
 #include "LiveLinkHubUEClientInfo.h"
+#include "Misc/Guid.h"
+#include "Misc/TextFilter.h"
 #include "Session/LiveLinkHubSessionManager.h"
 #include "SLiveLinkFilterSearchBox.h"
 #include "SPositiveActionButton.h"
diff --git a/Engine/Plugins/Animation/LiveLinkHub/Source/LiveLinkHub/Private/Session/LiveLinkHubSessionManager.h b/Engine/Plugins/Animation/LiveLinkHub/Source/LiveLinkHub/Private/Session/LiveLinkHubSessionManager.h
--- a/Engine/Plugins/Animation/LiveLinkHub/Source/LiveLinkHub/Private/Session/LiveLinkHubSessionManager.h
+++ b/Engine/Plugins/Animation/LiveLinkHub/Source/LiveLinkHub/Private/Session/LiveLinkHubSessionManager.h
@@ -400,7 +400,7 @@ public:
 	}	
 	//~ End LiveLinkHubSessionManager
 
-	//~ Begin ILiveLinkRecordingSessionInfo interface
+	//~ Begin ILiveLinkRecordingSession interface
 	virtual FString GetSessionName() const override
 	{
 		FScopeLock Lock(&CurrentSessionCS);
@@ -510,7 +510,48 @@ public:
 	{
 		return OnRecordingStoppedDelegate;
 	}
-	//~ End ILiveLinkRecordingSessionInfo interface
+
+	virtual bool CanRecord() const override
+	{
+		const FLiveLinkHubModule& HubModule = FModuleManager::Get().GetModuleChecked<FLiveLinkHubModule>("LiveLinkHub");
+
+		if (TSharedPtr<FLiveLinkHubRecordingController> RecordingController = HubModule.GetRecordingController())
+		{
+			return RecordingController->CanRecord();
+		}
+
+		return false;
+	}
+
+	virtual bool StartRecording() override
+	{
+		const FLiveLinkHubModule& HubModule = FModuleManager::Get().GetModuleChecked<FLiveLinkHubModule>("LiveLinkHub");
+
+		if (TSharedPtr<FLiveLinkHubRecordingController> RecordingController = HubModule.GetRecordingController())
+		{
+			RecordingController->StartRecording();
+			return RecordingController->IsRecording();
+		}
+
+		return false;
+	}
+
+	virtual bool StopRecording() override
+	{
+		const FLiveLinkHubModule& HubModule = FModuleManager::Get().GetModuleChecked<FLiveLinkHubModule>("LiveLinkHub");
+
+		if (TSharedPtr<FLiveLinkHubRecordingController> RecordingController = HubModule.GetRecordingController())
+		{
+			if (RecordingController->IsRecording())
+			{
+				RecordingController->StopRecording();
+				return !RecordingController->IsRecording();
+			}
+		}
+
+		return false;
+	}
+	//~ End ILiveLinkRecordingSession interface
 
 private:
 	/** Clear the hub data contained in the current session, resetting the hub to its default state. */
diff --git a/Engine/Plugins/Animation/LiveLinkHub/Source/LiveLinkHub/Public/ILiveLinkHubClientsModel.h b/Engine/Plugins/Animation/LiveLinkHub/Source/LiveLinkHub/Public/ILiveLinkHubClientsModel.h
new file mode 100644
--- /dev/null
+++ b/Engine/Plugins/Animation/LiveLinkHub/Source/LiveLinkHub/Public/ILiveLinkHubClientsModel.h
@@ -0,0 +1,158 @@
+// Copyright Epic Games, Inc. All Rights Reserved.
+
+#pragma once
+
+#include "Delegates/Delegate.h"
+#include "Containers/ObservableArray.h"
+#include "Features/IModularFeatures.h"
+#include "ILiveLinkHubClientsModel.generated.h"
+
+
+struct FLiveLinkHubAuxChannelRequestMessage;
+struct FLiveLinkClientInfoMessage;
+struct FLiveLinkHubUEClientInfo;
+struct FLiveLinkSubjectKey;
+class FMessageEndpoint;
+
+
+/**
+ * Identifier for a UE client receiving data from the hub.
+ */
+USTRUCT()
+struct FLiveLinkHubClientId
+{
+	GENERATED_BODY()
+
+	/** Default constructor, should only be used by the reflection code. */
+	FLiveLinkHubClientId() = default;
+
+	static FLiveLinkHubClientId NewId()
+	{
+		FLiveLinkHubClientId Id{};
+		Id.Guid = FGuid::NewGuid();
+		return Id;
+	}
+
+	bool IsValid() const
+	{
+		return Guid.IsValid();
+	}
+
+	/** Get the hash of this ID. */
+	friend uint32 GetTypeHash(FLiveLinkHubClientId Id)
+	{
+		return GetTypeHash(Id.Guid);
+	}
+
+	bool operator==(const FLiveLinkHubClientId& Other) const
+	{
+		return Guid == Other.Guid;
+	}
+
+	FString ToString() const
+	{
+		return Guid.ToString();
+	}
+
+private:
+	/** Unique identifier for this client. */
+	UPROPERTY()
+	FGuid Guid;
+};
+
+
+/** Exposes queries and operations for Unreal Editor clients connected to the local LLH instance. */
+class ILiveLinkHubClientsModel : public IModularFeature
+{
+public:
+	/** Types of client updates. */
+	enum class EClientEventType
+	{
+		Discovered, // UE client was discovered by the hub.
+		Disconnected, // Connection was lost with a client.
+		Connected, // Client goes from Discovered -> Connected
+		Reestablished, // Connection was restored with a previously disconnected client.
+		Modified // Information about a given client has changed (ie. level has changed).
+	};
+
+	DECLARE_TS_MULTICAST_DELEGATE_TwoParams(FOnClientEvent, FLiveLinkHubClientId, EClientEventType);
+
+	/** Name of the modular feature. */
+	static FName GetModularFeatureName()
+	{
+		static const FName ModularFeatureName("LiveLinkHubClientsModel");
+		return ModularFeatureName;
+	}
+
+	/** Convenience accessor for modular feature implementation. */
+	static ILiveLinkHubClientsModel* Get()
+	{
+		IModularFeatures& ModularFeatures = IModularFeatures::Get();
+		const TArray<ILiveLinkHubClientsModel*> FeatureImplementations = 
+			ModularFeatures.GetModularFeatureImplementations<ILiveLinkHubClientsModel>(GetModularFeatureName());
+		ensure(FeatureImplementations.Num() <= 1);
+		return FeatureImplementations.Num() > 0 ? FeatureImplementations[0] : nullptr;
+	}
+
+	/** Convenience accessor for modular feature implementation. */
+	static ILiveLinkHubClientsModel& GetChecked()
+	{
+		ILiveLinkHubClientsModel* MaybeFeature = Get();
+		check(MaybeFeature);
+		return *MaybeFeature;
+	}
+
+public:
+	virtual ~ILiveLinkHubClientsModel() = default;
+
+	/** Delegate used to get noticed about changes to the client list. */
+	virtual FOnClientEvent& OnClientEvent() = 0;
+
+	/** Get the status text of a client. */
+	virtual FText GetClientStatus(FLiveLinkHubClientId Client) const = 0;
+
+	/** Get the list of clients that should be displayed in the clients panel according to the current session config. */
+	virtual TArray<FLiveLinkHubClientId> GetSessionClients() const = 0;
+
+	/** Get the list of clients discovered by the hub. */
+	virtual TArray<FLiveLinkHubClientId> GetDiscoveredClients() const = 0;
+
+	/** Get information about a given client given its address. */
+	virtual TOptional<FLiveLinkHubUEClientInfo> GetClientInfo(FLiveLinkHubClientId InClient) const = 0;
+
+	/** Get the name of a client. */
+	virtual FText GetClientDisplayName(FLiveLinkHubClientId InClient) const = 0;
+
+	/** Send a connection message to this discovered client. */
+	virtual void ConnectTo(FLiveLinkHubClientId InClient) = 0;
+
+	/** Get whether a client should receive livelink data. */
+	virtual bool IsClientEnabled(FLiveLinkHubClientId InClient) const = 0;
+
+	/** Get whether a client is connected to the hub. */
+	virtual bool IsClientConnected(FLiveLinkHubClientId InClient) const = 0;
+
+	/** Set whether a client should receive livelink data. */
+	virtual void SetClientEnabled(FLiveLinkHubClientId InClient, bool bInEnable) = 0;
+
+	/** Get whether a subject is enabled on a given client. */
+	virtual bool IsSubjectEnabled(FLiveLinkHubClientId InClient, FName SubjectName) const = 0;
+
+	/** Set whether a subject should receive livelink data. */
+	virtual void SetSubjectEnabled(FLiveLinkHubClientId InClient, FName SubjectName, bool bInEnable) = 0;
+
+	/** Request to open an auxiliary endpoint with a client. */
+	virtual void RequestAuxiliaryChannel(
+		FLiveLinkHubClientId InClientId,
+		FMessageEndpoint& InAuxEndpoint,
+		UScriptStruct* InRequestTypeInfo,
+		FLiveLinkHubAuxChannelRequestMessage& InRequest
+	) = 0;
+
+	/** Request to open an auxiliary endpoint with a client. */
+	template<typename RequestType UE_REQUIRES(std::is_base_of_v<FLiveLinkHubAuxChannelRequestMessage, RequestType>)>
+	void RequestAuxiliaryChannel(FLiveLinkHubClientId InClientId, FMessageEndpoint& InAuxEndpoint, RequestType& InRequest)
+	{
+		RequestAuxiliaryChannel(InClientId, InAuxEndpoint, RequestType::StaticStruct(), InRequest);
+	}
+};
diff --git a/Engine/Plugins/Animation/LiveLinkHub/Source/LiveLinkHub/Public/ILiveLinkRecordingSession.h b/Engine/Plugins/Animation/LiveLinkHub/Source/LiveLinkHub/Public/ILiveLinkRecordingSession.h
new file mode 100644
--- /dev/null
+++ b/Engine/Plugins/Animation/LiveLinkHub/Source/LiveLinkHub/Public/ILiveLinkRecordingSession.h
@@ -0,0 +1,53 @@
+﻿// Copyright Epic Games, Inc. All Rights Reserved.
+
+#pragma once
+
+#include "Containers/StringFwd.h"
+#include "Delegates/Delegate.h"
+#include "Features/IModularFeatures.h"
+
+
+/** Provides an interface for querying and controlling the recording session. */
+class ILiveLinkRecordingSession : public IModularFeature
+{
+public:
+	DECLARE_MULTICAST_DELEGATE_OneParam(FOnSessionStringChanged, FStringView);
+	DECLARE_MULTICAST_DELEGATE_OneParam(FOnSessionIntChanged, int32);
+
+	/** Name of the modular feature. */
+	static FName GetModularFeatureName()
+	{
+		static const FName ModularFeatureName("LiveLinkRecordingSession");
+		return ModularFeatureName;
+	}
+
+	/** Convenience accessor for modular implementation. */
+	static ILiveLinkRecordingSession& Get()
+	{
+		IModularFeatures& ModularFeatures = IModularFeatures::Get();
+		check(ModularFeatures.IsModularFeatureAvailable(GetModularFeatureName()));
+		ensure(ModularFeatures.GetModularFeatureImplementationCount(GetModularFeatureName()) == 1);
+		return ModularFeatures.GetModularFeature<ILiveLinkRecordingSession>(GetModularFeatureName());
+	}
+
+	virtual FString GetSessionName() const = 0;
+	virtual FString GetSlateName() const = 0;
+	virtual int32 GetTakeNumber() const = 0;
+
+	virtual bool SetSessionName(FStringView InSessionName) = 0;
+	virtual bool SetSlateName(FStringView InSlateName) = 0;
+	virtual bool SetTakeNumber(int32 InTakeNumber) = 0;
+
+	virtual FOnSessionStringChanged& OnSessionNameChanged() = 0;
+	virtual FOnSessionStringChanged& OnSlateNameChanged() = 0;
+	virtual FOnSessionIntChanged& OnTakeNumberChanged() = 0;
+
+	virtual bool IsRecording() const = 0;
+
+	virtual FSimpleMulticastDelegate& OnRecordingStarted() = 0;
+	virtual FSimpleMulticastDelegate& OnRecordingStopped() = 0;
+
+	virtual bool CanRecord() const = 0;
+	virtual bool StartRecording() = 0;
+	virtual bool StopRecording() = 0;
+};
diff --git a/Engine/Plugins/Animation/LiveLinkHub/Source/LiveLinkHub/Public/ILiveLinkRecordingSessionInfo.h b/Engine/Plugins/Animation/LiveLinkHub/Source/LiveLinkHub/Public/ILiveLinkRecordingSessionInfo.h
--- a/Engine/Plugins/Animation/LiveLinkHub/Source/LiveLinkHub/Public/ILiveLinkRecordingSessionInfo.h
+++ b/Engine/Plugins/Animation/LiveLinkHub/Source/LiveLinkHub/Public/ILiveLinkRecordingSessionInfo.h
@@ -2,47 +2,7 @@
 
 #pragma once
 
-#include "Containers/StringFwd.h"
-#include "Delegates/Delegate.h"
-#include "Features/IModularFeatures.h"
+#include "ILiveLinkRecordingSession.h"
 
 
-/** Provides an interface for querying and setting recording session metadata. */
-class ILiveLinkRecordingSessionInfo : public IModularFeature
-{
-public:
-	DECLARE_MULTICAST_DELEGATE_OneParam(FOnSessionStringChanged, FStringView);
-	DECLARE_MULTICAST_DELEGATE_OneParam(FOnSessionIntChanged, int32);
-
-	/** Name of the modular feature. */
-	static FName GetModularFeatureName()
-	{
-		static const FName ModularFeatureName("LiveLinkRecordingSessionInfo");
-		return ModularFeatureName;
-	}
-
-	/** Convenience accessor for modular implementation. */
-	static ILiveLinkRecordingSessionInfo& Get()
-	{
-		IModularFeatures& ModularFeatures = IModularFeatures::Get();
-		ensure(ModularFeatures.GetModularFeatureImplementationCount(GetModularFeatureName()) == 1);
-		return ModularFeatures.GetModularFeature<ILiveLinkRecordingSessionInfo>(GetModularFeatureName());
-	}
-
-	virtual FString GetSessionName() const = 0;
-	virtual FString GetSlateName() const = 0;
-	virtual int32 GetTakeNumber() const = 0;
-
-	virtual bool SetSessionName(FStringView InSessionName) = 0;
-	virtual bool SetSlateName(FStringView InSlateName) = 0;
-	virtual bool SetTakeNumber(int32 InTakeNumber) = 0;
-
-	virtual FOnSessionStringChanged& OnSessionNameChanged() = 0;
-	virtual FOnSessionStringChanged& OnSlateNameChanged() = 0;
-	virtual FOnSessionIntChanged& OnTakeNumberChanged() = 0;
-
-	virtual bool IsRecording() const = 0;
-
-	virtual FSimpleMulticastDelegate& OnRecordingStarted() = 0;
-	virtual FSimpleMulticastDelegate& OnRecordingStopped() = 0;
-};
+using ILiveLinkRecordingSessionInfo = ILiveLinkRecordingSession;
diff --git a/Engine/Plugins/Animation/LiveLinkHub/Source/LiveLinkHubEditor/Private/LiveLinkHubClientIdCustomization.cpp b/Engine/Plugins/Animation/LiveLinkHub/Source/LiveLinkHubEditor/Private/LiveLinkHubClientIdCustomization.cpp
new file mode 100644
--- /dev/null
+++ b/Engine/Plugins/Animation/LiveLinkHub/Source/LiveLinkHubEditor/Private/LiveLinkHubClientIdCustomization.cpp
@@ -0,0 +1,204 @@
+// Copyright Epic Games, Inc. All Rights Reserved.
+
+#include "LiveLinkHubClientIdCustomization.h"
+#include "DetailLayoutBuilder.h"
+#include "DetailWidgetRow.h"
+#include "Framework/Application/SlateApplication.h"
+#include "Framework/Commands/UIAction.h"
+#include "Framework/MultiBox/MultiBoxBuilder.h"
+#include "IDetailChildrenBuilder.h"
+#include "ILiveLinkHubClientsModel.h"
+#include "Input/Events.h"
+#include "Layout/Margin.h"
+#include "Layout/Visibility.h"
+#include "Misc/Attribute.h"
+#include "PropertyHandle.h"
+#include "Textures/SlateIcon.h"
+#include "Widgets/DeclarativeSyntaxSupport.h"
+#include "Widgets/SBoxPanel.h"
+#include "Widgets/SCompoundWidget.h"
+#include "Widgets/SWidget.h"
+#include "Widgets/Input/SComboButton.h"
+#include "Widgets/Text/STextBlock.h"
+
+
+#define LOCTEXT_NAMESPACE "LiveLinkHubClientIdCustomization"
+
+
+class SLiveLinkHubClientIdDropdown : public SCompoundWidget
+{
+public:
+	DECLARE_DELEGATE_OneParam(FOnValueChanged, FLiveLinkHubClientId);
+
+public:
+	SLATE_BEGIN_ARGS(SLiveLinkHubClientIdDropdown)
+		: _MinDesiredValueWidth( 40 )
+		{
+		}
+
+		SLATE_ARGUMENT(TSharedPtr<IPropertyHandle>, PropertyHandle)
+
+		/** Controls the minimum width for the text box portion of the control. */
+		SLATE_ATTRIBUTE(float, MinDesiredValueWidth)
+
+		/** The callback for when the value changes. */
+		SLATE_EVENT(FOnValueChanged, OnValueChanged)
+
+	SLATE_END_ARGS()
+
+	void Construct(const FArguments& InArgs)
+	{
+		PropertyHandle = InArgs._PropertyHandle;
+		OnValueChanged = InArgs._OnValueChanged;
+
+		ChildSlot
+		[
+			SNew(SVerticalBox)
+			+ SVerticalBox::Slot()
+			.AutoHeight()
+			[
+				SNew(SHorizontalBox)
+				+ SHorizontalBox::Slot()
+				[
+					SNew(SComboButton)
+					.ContentPadding(1)
+					.OnGetMenuContent(this, &SLiveLinkHubClientIdDropdown::BuildMenu)
+					.ButtonContent()
+					[
+						SNew(STextBlock)
+						.MinDesiredWidth(InArgs._MinDesiredValueWidth)
+						.Font(IDetailLayoutBuilder::GetDetailFont())
+						.Text(this, &SLiveLinkHubClientIdDropdown::GetValueText)
+					]
+				]
+			]
+		];
+	}
+
+private:
+	TSharedRef<SWidget> BuildMenu()
+	{
+		FMenuBuilder MenuBuilder(true, nullptr);
+
+		if (ILiveLinkHubClientsModel* ClientsModel = ILiveLinkHubClientsModel::Get())
+		{
+			TArray<FLiveLinkHubClientId> SessionClients = ClientsModel->GetSessionClients();
+			for (FLiveLinkHubClientId ClientId : SessionClients)
+			{
+				FUIAction MenuAction(FExecuteAction::CreateSP(this, &SLiveLinkHubClientIdDropdown::SetValue, ClientId));
+				MenuBuilder.AddMenuEntry(ClientsModel->GetClientDisplayName(ClientId), FText::GetEmpty(), FSlateIcon(), MenuAction);
+			}
+		}
+
+		return MenuBuilder.MakeWidget();
+	}
+
+	void SetValue(FLiveLinkHubClientId InNewClientId)
+	{
+		if (PropertyHandle && PropertyHandle->IsValidHandle())
+		{
+			FStructProperty* StructProperty = CastFieldChecked<FStructProperty>(PropertyHandle->GetProperty());
+
+			TArray<void*> RawData;
+			PropertyHandle->AccessRawData(RawData);
+			if (RawData.Num() > 0 && RawData[0])
+			{
+				FLiveLinkHubClientId* PreviousValue = reinterpret_cast<FLiveLinkHubClientId*>(RawData[0]);
+
+				FString TextValue;
+				StructProperty->Struct->ExportText(TextValue, &InNewClientId, PreviousValue, nullptr, EPropertyPortFlags::PPF_None, nullptr);
+				ensure(PropertyHandle->SetValueFromFormattedString(TextValue, EPropertyValueSetFlags::DefaultFlags) == FPropertyAccess::Result::Success);
+			}
+		}
+	}
+
+	FText GetValueText() const
+	{
+		if (PropertyHandle && PropertyHandle->IsValidHandle())
+		{
+			TArray<const void*> RawData;
+			PropertyHandle->AccessRawData(RawData);
+			if (RawData.Num() == 1 && RawData[0])
+			{
+				return GetClientText(*(const FLiveLinkHubClientId*)RawData[0]);
+			}
+			else
+			{
+				return LOCTEXT("ValueText_MultipleValues", "Multiple Values");
+			}
+		}
+
+		return LOCTEXT("ValueText_HandleInvalid", "(Error: Invalid property handle)");
+	}
+
+	static FText GetClientText(FLiveLinkHubClientId InClientId)
+	{
+		if (InClientId.IsValid())
+		{
+			if (ILiveLinkHubClientsModel* ClientsModel = ILiveLinkHubClientsModel::Get())
+			{
+				return ClientsModel->GetClientDisplayName(InClientId);
+			}
+			else
+			{
+				return LOCTEXT("ClientText_ClientsModelInvalid", "(Error: No clients model)");
+			}
+		}
+		else
+		{
+			return LOCTEXT("ClientText_IdInvalid", "(None)");
+		}
+	}
+
+private:
+	TSharedPtr<IPropertyHandle> PropertyHandle;
+	FOnValueChanged OnValueChanged;
+};
+
+
+//////////////////////////////////////////////////////////////////////////
+
+
+TSharedRef<IPropertyTypeCustomization> FLiveLinkHubClientIdCustomization::MakeInstance()
+{
+	return MakeShared<FLiveLinkHubClientIdCustomization>();
+}
+
+
+void FLiveLinkHubClientIdCustomization::CustomizeHeader(
+	TSharedRef<IPropertyHandle> InStructPropertyHandle,
+	FDetailWidgetRow& HeaderRow,
+	IPropertyTypeCustomizationUtils& StructCustomizationUtils
+)
+{
+	StructPropertyHandle = InStructPropertyHandle;
+	FString CurrentValue;
+	StructPropertyHandle->GetValueAsFormattedString(CurrentValue);
+
+	HeaderRow
+	.NameContent()
+	[
+		StructPropertyHandle->CreatePropertyNameWidget()
+	]
+	.ValueContent()
+	[
+		SNew(SBox)
+		.Padding(FMargin(0,2,0,1))
+		[
+			SNew(SLiveLinkHubClientIdDropdown)
+			.PropertyHandle(StructPropertyHandle)
+		]
+	];
+}
+
+
+void FLiveLinkHubClientIdCustomization::CustomizeChildren(
+	TSharedRef<IPropertyHandle> InStructPropertyHandle,
+	IDetailChildrenBuilder& ChildBuilder,
+	IPropertyTypeCustomizationUtils& StructCustomizationUtils
+)
+{
+}
+
+
+#undef LOCTEXT_NAMESPACE
diff --git a/Engine/Plugins/Animation/LiveLinkHub/Source/LiveLinkHubEditor/Private/LiveLinkHubClientIdCustomization.h b/Engine/Plugins/Animation/LiveLinkHub/Source/LiveLinkHubEditor/Private/LiveLinkHubClientIdCustomization.h
new file mode 100644
--- /dev/null
+++ b/Engine/Plugins/Animation/LiveLinkHub/Source/LiveLinkHubEditor/Private/LiveLinkHubClientIdCustomization.h
@@ -0,0 +1,25 @@
+// Copyright Epic Games, Inc. All Rights Reserved.
+
+#pragma once
+
+#include "IPropertyTypeCustomization.h"
+
+
+class IPropertyHandle;
+class FDetailWidgetRow;
+class IDetailChildrenBuilder;
+
+
+class FLiveLinkHubClientIdCustomization : public IPropertyTypeCustomization
+{
+public:
+	static TSharedRef<IPropertyTypeCustomization> MakeInstance();
+
+	virtual void CustomizeHeader(TSharedRef<IPropertyHandle> StructPropertyHandle, FDetailWidgetRow& HeaderRow, IPropertyTypeCustomizationUtils& StructCustomizationUtils) override;
+	virtual void CustomizeChildren(TSharedRef<IPropertyHandle> InStructPropertyHandle, IDetailChildrenBuilder& ChildBuilder, IPropertyTypeCustomizationUtils& StructCustomizationUtils) override;
+	virtual bool ShouldInlineKey() const override { return true; }
+
+private:	
+	/** Cached property handle */
+	TSharedPtr<IPropertyHandle> StructPropertyHandle;
+};
diff --git a/Engine/Plugins/Animation/LiveLinkHub/Source/LiveLinkHubEditor/Private/LiveLinkHubEditorModule.cpp b/Engine/Plugins/Animation/LiveLinkHub/Source/LiveLinkHubEditor/Private/LiveLinkHubEditorModule.cpp
--- a/Engine/Plugins/Animation/LiveLinkHub/Source/LiveLinkHubEditor/Private/LiveLinkHubEditorModule.cpp
+++ b/Engine/Plugins/Animation/LiveLinkHub/Source/LiveLinkHubEditor/Private/LiveLinkHubEditorModule.cpp
@@ -6,11 +6,13 @@
 #include "HAL/FileManager.h"
 #include "ILauncherPlatform.h"
 #include "LauncherPlatformModule.h"
+#include "LiveLinkHubClientIdCustomization.h"
 #include "LiveLinkHubLauncherUtils.h"
 #include "LiveLinkHubEditorSettings.h"
 #include "Misc/App.h"
 #include "Misc/AsyncTaskNotification.h"
 #include "Misc/MessageDialog.h"
+#include "PropertyEditorModule.h"
 #include "Runtime/Launch/Resources/Version.h"
 #include "SLiveLinkHubEditorStatusBar.h"
 #include "ToolMenus.h"
@@ -41,6 +43,11 @@ void FLiveLinkHubEditorModule::ShutdownModule()
 	{
 		FCoreDelegates::OnPostEngineInit.RemoveAll(this);
 		UnregisterLiveLinkHubStatusBar();
+
+		if (FPropertyEditorModule* PropertyModule = FModuleManager::GetModulePtr<FPropertyEditorModule>("PropertyEditor"))
+		{
+			PropertyModule->UnregisterCustomPropertyTypeLayout("LiveLinkHubClientId");
+		}
 	}
 }
 
@@ -59,6 +66,10 @@ void FLiveLinkHubEditorModule::OnPostEngineInit()
 			LOCTEXT("LiveLinkHubTooltip", "Launch the Live Link Hub app."),
 			FSlateIcon("LiveLinkStyle", "LiveLinkClient.Common.Icon.Small"),
 			FUIAction(FExecuteAction::CreateStatic(&UE::LiveLinkHubLauncherUtils::OpenLiveLinkHub)));
+
+		FPropertyEditorModule& PropertyModule = FModuleManager::LoadModuleChecked<FPropertyEditorModule>("PropertyEditor");
+		PropertyModule.RegisterCustomPropertyTypeLayout("LiveLinkHubClientId", FOnGetPropertyTypeCustomizationInstance::CreateStatic(&FLiveLinkHubClientIdCustomization::MakeInstance));
+		PropertyModule.NotifyCustomizationModuleChanged();
 	}
 }
 
diff --git a/Engine/Plugins/Animation/LiveLinkHub/Source/LiveLinkHubMessaging/Private/LiveLinkHubControlChannel.h b/Engine/Plugins/Animation/LiveLinkHub/Source/LiveLinkHubMessaging/Private/LiveLinkHubControlChannel.h
--- a/Engine/Plugins/Animation/LiveLinkHub/Source/LiveLinkHubMessaging/Private/LiveLinkHubControlChannel.h
+++ b/Engine/Plugins/Animation/LiveLinkHub/Source/LiveLinkHubMessaging/Private/LiveLinkHubControlChannel.h
@@ -119,7 +119,20 @@ public:
 		Builder.Handling<FLiveLinkHubDiscoveryMessage>(GameThreadDispatch)
 			.Handling<FLiveLinkHubCustomTimeStepSettings>(GameThreadDispatch)
 			.Handling<FLiveLinkHubTimecodeSettings>(GameThreadDispatch)
-			.Handling<FLiveLinkHubDisconnectMessage>(GameThreadDispatch);
+			.Handling<FLiveLinkHubDisconnectMessage>(GameThreadDispatch)
+			.WithCatchall(
+				[this, GameThreadDispatch]
+				(const TSharedRef<IMessageContext, ESPMode::ThreadSafe>& InContext) mutable
+				{
+					const UScriptStruct* MessageTypeInfo = InContext->GetMessageTypeInfo().Get();
+
+					// Handle subclasses of aux request.
+			        if (MessageTypeInfo->IsChildOf(FLiveLinkHubAuxChannelRequestMessage::StaticStruct()))
+					{
+				        GameThreadDispatch(*(const FLiveLinkHubAuxChannelRequestMessage*)InContext->GetMessage(), InContext);
+					}
+				})
+			;
 	}
 
 	/** Informs a LLH instance that this endpoint exists in order to create an entry in their MessageBus AddressBook. */
@@ -189,6 +202,10 @@ public:
 		return ClientInfo;
 	}
 
+	/** Returning false from this delegate will respond to the sender with FLiveLinkHubAuxChannelRejectMessage. */
+	using FAuxRequestDelegate = TTSDelegate<bool(const FLiveLinkHubAuxChannelRequestMessage& InMessage, const TSharedRef<IMessageContext, ESPMode::ThreadSafe>& InContext)>;
+	FAuxRequestDelegate OnAuxRequest;
+
 private:
 	/** Unregister callbacks. */
 	void Shutdown()
@@ -288,6 +305,25 @@ private:
 		}
 	}
 
+	/** Handle aux channel requests by delegating to the provided callback. */
+	template<>
+	void Handle(const FLiveLinkHubAuxChannelRequestMessage& InMessage, const TSharedRef<IMessageContext, ESPMode::ThreadSafe>& InContext)
+	{
+		if (OnAuxRequest.IsBound())
+		{
+			const bool bHandled = OnAuxRequest.Execute(InMessage, InContext);
+			if (bHandled)
+			{
+				return;
+			}
+		}
+
+		// Request not handled; respond with reject.
+		FLiveLinkHubAuxChannelRejectMessage* RejectMessage = FMessageEndpoint::MakeMessage<FLiveLinkHubAuxChannelRejectMessage>();
+		RejectMessage->ChannelId = InMessage.ChannelId;
+		Endpoint->Send(RejectMessage, EMessageFlags::Reliable, nullptr, { InContext->GetSender() }, FTimespan::Zero(), FDateTime::MaxValue());
+	}
+
 #if WITH_EDITOR
 	/** Handler called on map changed to update the livelink hub. */
 	void OnMapChanged(UWorld* World, EMapChangeType ChangeType)
diff --git a/Engine/Plugins/Animation/LiveLinkHub/Source/LiveLinkHubMessaging/Private/LiveLinkHubMessagingModule.cpp b/Engine/Plugins/Animation/LiveLinkHub/Source/LiveLinkHubMessaging/Private/LiveLinkHubMessagingModule.cpp
--- a/Engine/Plugins/Animation/LiveLinkHub/Source/LiveLinkHubMessaging/Private/LiveLinkHubMessagingModule.cpp
+++ b/Engine/Plugins/Animation/LiveLinkHub/Source/LiveLinkHubMessaging/Private/LiveLinkHubMessagingModule.cpp
@@ -9,6 +9,7 @@
 #include "LiveLinkHubMessageBusSourceFactory.h"
 #include "LiveLinkMessageBusFinder.h"
 #include "LiveLinkMessageBusSourceFactory.h"
+#include "Logging/StructuredLog.h"
 #include "MessageEndpoint.h"
 #include "MessageEndpointBuilder.h"
 #include "Misc/CoreDelegates.h"
@@ -53,6 +54,7 @@ void FLiveLinkHubMessagingModule::StartupModule()
 	);
 
 	ControlChannel = MakeShared<FLiveLinkHubControlChannel>(FLiveLinkHubControlChannel::EChannelMode::Global);
+	ControlChannel->OnAuxRequest.BindRaw(this, &FLiveLinkHubMessagingModule::OnAuxRequest);
 	FCoreDelegates::OnPostEngineInit.AddLambda([this]() { ControlChannel->Initialize(); });
 
 #endif
@@ -72,6 +74,7 @@ void FLiveLinkHubMessagingModule::ShutdownModule()
 	UnregisterSettings();
 
 #if WITH_LIVELINK_DISCOVERY_MANAGER_THREAD
+	ControlChannel->OnAuxRequest.Unbind();
 	ControlChannel.Reset();
 	ConnectionManager.Reset();
 #endif
@@ -111,6 +114,45 @@ ELiveLinkTopologyMode FLiveLinkHubMessagingModule::GetHostTopologyMode() const
 	return InstanceInfo.TopologyMode;
 }
 
+
+bool FLiveLinkHubMessagingModule::RegisterAuxChannelRequestHandler(UScriptStruct* InRequestTypeStruct, FAuxChannelRequestHandlerFunc&& InHandlerFunc)
+{
+	if (!ensure(InRequestTypeStruct))
+	{
+		UE_LOGFMT(LogLiveLinkHubMessaging, Error, "Attempted to register handler for null request type");
+		return false;
+	}
+
+	if (!ensure(InRequestTypeStruct != FLiveLinkHubAuxChannelRequestMessage::StaticStruct()))
+	{
+		UE_LOGFMT(LogLiveLinkHubMessaging, Error, "Attempted to register handler for base request type");
+		return false;
+	}
+
+	if (!ensure(!AuxChannelRequestHandlers.Find(InRequestTypeStruct)))
+	{
+		UE_LOGFMT(LogLiveLinkHubMessaging, Error, "Handler already exists for request type {RequestTypeName}", InRequestTypeStruct->GetName());
+		return false;
+	}
+
+	AuxChannelRequestHandlers.Add(InRequestTypeStruct, MoveTemp(InHandlerFunc));
+
+	return true;
+}
+
+
+bool FLiveLinkHubMessagingModule::UnregisterAuxChannelRequestHandler(UScriptStruct* InRequestTypeStruct)
+{
+	if (!ensure(InRequestTypeStruct))
+	{
+		UE_LOGFMT(LogLiveLinkHubMessaging, Error, "Attempted to unregister null request type");
+		return false;
+	}
+
+	return AuxChannelRequestHandlers.Remove(InRequestTypeStruct) != 0;
+}
+
+
 bool FLiveLinkHubMessagingModule::OnFilterMessageBusSource(UClass* FactoryClass, TSharedPtr<struct FProviderPollResult, ESPMode::ThreadSafe> PollResult)
 {
 	// Only display Hub/Spoke sources in "LiveLinkHub" section of the add source dropdown.
@@ -132,6 +174,25 @@ TSharedPtr<FLiveLinkHubControlChannel> FLiveLinkHubMessagingModule::GetControlCh
 	return ControlChannel;
 }
 
+bool FLiveLinkHubMessagingModule::OnAuxRequest(
+    const FLiveLinkHubAuxChannelRequestMessage& InMessage,
+	const TSharedRef<IMessageContext, ESPMode::ThreadSafe>& InContext
+)
+{
+	const UScriptStruct* MessageTypeInfo = InContext->GetMessageTypeInfo().Get();
+	if (FAuxChannelRequestHandlerFunc* Handler = AuxChannelRequestHandlers.Find(MessageTypeInfo))
+	{
+		(*Handler)(InMessage, InContext);
+		return true;
+	}
+	else
+	{
+		// Returning false generates a FLiveLinkHubAuxChannelRejectMessage response.
+		return false;
+	}
+}
+
+
 void FLiveLinkHubMessagingModule::RegisterSettings()
 {
 #if WITH_EDITOR
diff --git a/Engine/Plugins/Animation/LiveLinkHub/Source/LiveLinkHubMessaging/Private/LiveLinkHubMessagingModule.h b/Engine/Plugins/Animation/LiveLinkHub/Source/LiveLinkHubMessaging/Private/LiveLinkHubMessagingModule.h
--- a/Engine/Plugins/Animation/LiveLinkHub/Source/LiveLinkHubMessaging/Private/LiveLinkHubMessagingModule.h
+++ b/Engine/Plugins/Animation/LiveLinkHub/Source/LiveLinkHubMessaging/Private/LiveLinkHubMessagingModule.h
@@ -14,6 +14,7 @@
 class FLiveLinkHubConnectionManager;
 class FLiveLinkHubMessageBusSource;
 
+
 class FLiveLinkHubMessagingModule : public ILiveLinkHubMessagingModule
 {
 public:
@@ -33,10 +34,13 @@ public:
 	virtual FLiveLinkHubInstanceId GetInstanceId() const override;
 	virtual void SetInstanceId(const FLiveLinkHubInstanceId& Id) override;
 	virtual ELiveLinkTopologyMode GetHostTopologyMode() const override;
+	virtual bool RegisterAuxChannelRequestHandler(UScriptStruct* InRequestTypeStruct, FAuxChannelRequestHandlerFunc&& InHandlerFunc) override;
+	virtual bool UnregisterAuxChannelRequestHandler(UScriptStruct* InRequestTypeStruct) override;
 	//~ End ILiveLinkHubMessagingModule interface
 
 	/** Used by the LiveLinkHubMessageBusSource to acquire a control channel when it's created. */
 	TSharedPtr<class FLiveLinkHubControlChannel> GetControlChannel();
+
 private:
 	/**
 	 * Note: Invoked on the UI (Game) thread.
@@ -47,6 +51,9 @@ private:
 	/** Handle a discovery request from LLH to populate its address book with the ControlEndpoint's address. */
 	void OnDiscoveryRequest(const struct FMessageAddress& RemoteAddress) const;
 
+	/** Handle aux channel request messages. */
+	bool OnAuxRequest(const FLiveLinkHubAuxChannelRequestMessage& InMessage, const TSharedRef<IMessageContext, ESPMode::ThreadSafe>& InContext);
+
 	/** Register settings related to LiveLinkHub messaging. */
 	void RegisterSettings();
 
@@ -80,5 +87,7 @@ private:
 
 	/** Control channel used to communicate with LLH Providers built on 5.7+. */
 	TSharedPtr<class FLiveLinkHubControlChannel> ControlChannel;
+
+	TMap<const UScriptStruct*, FAuxChannelRequestHandlerFunc> AuxChannelRequestHandlers;
 };
 
diff --git a/Engine/Plugins/Animation/LiveLinkHub/Source/LiveLinkHubMessaging/Public/ILiveLinkHubMessagingModule.h b/Engine/Plugins/Animation/LiveLinkHub/Source/LiveLinkHubMessaging/Public/ILiveLinkHubMessagingModule.h
--- a/Engine/Plugins/Animation/LiveLinkHub/Source/LiveLinkHubMessaging/Public/ILiveLinkHubMessagingModule.h
+++ b/Engine/Plugins/Animation/LiveLinkHub/Source/LiveLinkHubMessaging/Public/ILiveLinkHubMessagingModule.h
@@ -12,9 +12,13 @@
 
 #define UE_API LIVELINKHUBMESSAGING_API
 
-DECLARE_MULTICAST_DELEGATE_OneParam(FOnHubConnectionEstablished, FGuid SourceId);
+enum class ELiveLinkTopologyMode : uint8;
+struct FLiveLinkHubAuxChannelRequestMessage;
+class IMessageContext;
 
 
+DECLARE_MULTICAST_DELEGATE_OneParam(FOnHubConnectionEstablished, FGuid SourceId);
+
 DECLARE_LOG_CATEGORY_EXTERN(LogLiveLinkHubMessaging, Display, All);
 
 
@@ -31,7 +35,6 @@ private:
 	TVariant<FGuid, FString> Id;
 };
 
-enum class ELiveLinkTopologyMode : uint8;
 
 class ILiveLinkHubMessagingModule : public IModuleInterface
 {
@@ -50,6 +53,44 @@ public:
 
 	/** Get the ID for this running instance. Returns an empty guid if this isn't running as part of LiveLinkHub. */
 	virtual FLiveLinkHubInstanceId GetInstanceId() const = 0;
+
+	using FAuxChannelRequestHandlerFunc = TUniqueFunction<void(
+		const FLiveLinkHubAuxChannelRequestMessage&,
+		const TSharedRef<IMessageContext, ESPMode::ThreadSafe>&
+	)>;
+
+	/** Register a callback to be invoked when a peer attempts to negotiate an auxiliary endpoint. */
+	template<typename RequestType UE_REQUIRES(std::is_base_of_v<FLiveLinkHubAuxChannelRequestMessage, RequestType>)>
+	bool RegisterAuxChannelRequestHandler(TUniqueFunction<void(const RequestType&, const TSharedRef<IMessageContext, ESPMode::ThreadSafe>& Context)>&& InHandlerFunc)
+	{
+		// Type erase the handler and pass through to the virtual helper.
+		return RegisterAuxChannelRequestHandler(
+		    RequestType::StaticStruct(),
+			[HandlerFunc = MoveTemp(InHandlerFunc)]
+			(
+				const FLiveLinkHubAuxChannelRequestMessage& InMessage,
+				const TSharedRef<IMessageContext, ESPMode::ThreadSafe>& InContext
+			)
+			{
+			    const RequestType& Message = static_cast<const RequestType&>(InMessage);
+				HandlerFunc(Message, InContext);
+			}
+		);
+	}
+
+	/** Unregister an existing auxiliary endpoint handler. */
+	template<typename RequestType UE_REQUIRES(std::is_base_of_v<FLiveLinkHubAuxChannelRequestMessage, RequestType>)>
+	bool UnregisterAuxChannelRequestHandler()
+	{
+		return UnregisterAuxChannelRequestHandler(RequestType::StaticStruct());
+	}
+
+	/** Register a callback to be invoked when a peer attempts to negotiate an auxiliary endpoint. */
+	virtual bool RegisterAuxChannelRequestHandler(UScriptStruct* InRequestTypeStruct, FAuxChannelRequestHandlerFunc&& InHandlerFunc) = 0;
+
+	/** Unregister an existing auxiliary endpoint handler. */
+	virtual bool UnregisterAuxChannelRequestHandler(UScriptStruct* InRequestTypeStruct) = 0;
+
 };
 
 #undef UE_API
diff --git a/Engine/Plugins/Animation/LiveLinkHub/Source/LiveLinkHubMessaging/Public/LiveLinkHubMessages.h b/Engine/Plugins/Animation/LiveLinkHub/Source/LiveLinkHubMessaging/Public/LiveLinkHubMessages.h
--- a/Engine/Plugins/Animation/LiveLinkHub/Source/LiveLinkHubMessaging/Public/LiveLinkHubMessages.h
+++ b/Engine/Plugins/Animation/LiveLinkHub/Source/LiveLinkHubMessaging/Public/LiveLinkHubMessages.h
@@ -286,6 +286,57 @@ struct FLiveLinkHubDiscoveryMessage
 	int32 DiscoveryProtocolVersion = 1;
 };
 
+/** Base for control messages for negotiating auxiliary channel endpoints. */
+USTRUCT()
+struct FLiveLinkHubAuxChannelBaseMessage
+{
+	GENERATED_BODY()
+
+	/** Aux channel version. */
+	UPROPERTY()
+	int32 AuxProtocolVersion = 1;
+
+	/** Identifier for the channel. */
+	UPROPERTY()
+	FGuid ChannelId;
+};
+
+/**
+ * Sent to LLH control endpoint to negotiate additional auxiliary endpoints.
+ * Subclass this struct to specialize it for your channel type, then:
+ * 1) Register a handler for your derived channel request type via
+ *    ILiveLinkHubMessagingModule::RegisterAuxChannelRequestHandler.
+ * 2) Send a channel request via ILiveLinkHubClientsModel::RequestAuxiliaryChannel.
+ */
+USTRUCT()
+struct FLiveLinkHubAuxChannelRequestMessage : public FLiveLinkHubAuxChannelBaseMessage
+{
+	GENERATED_BODY()
+};
+
+/** Send this from your aux channel handler in response to an accepted request. */
+USTRUCT()
+struct FLiveLinkHubAuxChannelAcceptMessage : public FLiveLinkHubAuxChannelBaseMessage
+{
+	GENERATED_BODY()
+};
+
+/**
+ * Send this from your aux channel handler in response to a rejected request.
+ * Also sent by the framework if no handler was registered for the request type.
+ */
+USTRUCT()
+struct FLiveLinkHubAuxChannelRejectMessage : public FLiveLinkHubAuxChannelBaseMessage
+{
+	GENERATED_BODY()
+};
+
+/** Send this to close an existing aux channel. */
+USTRUCT()
+struct FLiveLinkHubAuxChannelCloseMessage : public FLiveLinkHubAuxChannelBaseMessage
+{
+	GENERATED_BODY()
+};
 
 
 #undef UE_API
