LiveLinkHub - Add Client filtering support

- Move autoconnect behavior from settings to user settings.
- Add ability to set filters for client autoconnections
- Add ability to save current set of connected clients to a filter preset (and ability to load/delete them)

diff --git a/Engine/Plugins/Animation/LiveLinkHub/Source/LiveLinkHub/Private/Clients/LiveLinkHubProvider.cpp b/Engine/Plugins/Animation/LiveLinkHub/Source/LiveLinkHub/Private/Clients/LiveLinkHubProvider.cpp
--- a/Engine/Plugins/Animation/LiveLinkHub/Source/LiveLinkHub/Private/Clients/LiveLinkHubProvider.cpp
+++ b/Engine/Plugins/Animation/LiveLinkHub/Source/LiveLinkHub/Private/Clients/LiveLinkHubProvider.cpp
@@ -94,7 +94,12 @@ FLiveLinkHubProvider::FLiveLinkHubProvider(const TSharedRef<ILiveLinkHubSessionM
 	FCoreDelegates::OnPostEngineInit.AddLambda([this]()
 	{
 		const double ValidateConnectionsRate = GetDefault<ULiveLinkSettings>()->MessageBusPingRequestFrequency;
-        GEditor->GetTimerManager()->SetTimer(ValidateConnectionsTimer, FTimerDelegate::CreateRaw(this, &FLiveLinkHubProvider::ValidateHubConnections), ValidateConnectionsRate, true);
+		if (GEditor)
+		{
+			GEditor->GetTimerManager()->SetTimer(ValidateConnectionsTimer, FTimerDelegate::CreateRaw(this, &FLiveLinkHubProvider::ValidateHubConnections), ValidateConnectionsRate, true);
+		}
+
+		GetMutableDefault<ULiveLinkHubUserSettings>()->OnFiltersModifiedDelegate.AddRaw(this, &FLiveLinkHubProvider::OnFiltersModified);
 	});
 
 	IModularFeatures::Get().RegisterModularFeature(ILiveLinkHubClientsModel::GetModularFeatureName(), this);
@@ -102,6 +107,11 @@ FLiveLinkHubProvider::FLiveLinkHubProvider(const TSharedRef<ILiveLinkHubSessionM
 
 FLiveLinkHubProvider::~FLiveLinkHubProvider()
 {
+	if (UObjectInitialized() && !IsEngineExitRequested())
+	{
+		GetMutableDefault<ULiveLinkHubUserSettings>()->OnFiltersModifiedDelegate.RemoveAll(this);
+	}
+
 	if (GEditor)
 	{
 		GEditor->GetTimerManager()->ClearTimer(ValidateConnectionsTimer);
@@ -488,7 +498,7 @@ void FLiveLinkHubProvider::HandleHubConnectMessage(const FLiveLinkHubConnectMess
 
 		if (!ClientId.IsValid())
 		{
-			ClientId = UpdateAddressBookEntry(ControlEndpoint, DiscoveryProtocolVersion, Message.ClientInfo.TopologyMode, Message.ClientInfo.Hostname);
+			ClientId = UpdateAddressBookEntry(ControlEndpoint, DiscoveryProtocolVersion, Message.ClientInfo.TopologyMode, Message.ClientInfo.Hostname, Message.ClientInfo.ProjectName);
 
 			if (DiscoveryProtocolVersion > 1)
 			{
@@ -749,7 +759,7 @@ void FLiveLinkHubProvider::ValidateHubConnections()
 					{
 						It->LastDiscoveryRequestTimestamp = FPlatformTime::Seconds();
 
-						if (ShouldAutoConnectTo(It->Hostname) 
+						if (ShouldAutoConnectTo(It->Hostname, It->IP, It->ProjectName) 
 							&& (GetDefault<ULiveLinkHubMessagingSettings>()->CanTransmitTo(FLiveLinkHub::Get()->GetTopologyMode(), It->TopologyMode))) // Previous versions of UE/LLH relied on the client side to check for Hub/Spoke compatibility.
 						{
 							It->bSentDiscoveryRequest = true;
@@ -799,7 +809,7 @@ TMap<FName, FString> FLiveLinkHubProvider::GetAnnotations() const
 	TMap<FName, FString> AnnotationsCopy = Annotations;
 
 	// Make a local copy of the annotations map and add the autoconnect annotation since this value might change over time.
-	const FString AutoConnectValue = StaticEnum<ELiveLinkHubAutoConnectMode>()->GetNameStringByValue(static_cast<int64>(GetDefault<ULiveLinkHubSettings>()->AutoConnectClients));
+	const FString AutoConnectValue = StaticEnum<ELiveLinkHubAutoConnectMode>()->GetNameStringByValue(static_cast<int64>(GetDefault<ULiveLinkHubUserSettings>()->GetAutoConnectMode()));
 	AnnotationsCopy.Add(FLiveLinkHubMessageAnnotation::AutoConnectModeAnnotation, AutoConnectValue);
 
 	// In case this is handled by the discovery manager (which doesn't directly handle LiveLinkHub messages)
@@ -1052,7 +1062,7 @@ void FLiveLinkHubProvider::RequestAuxiliaryChannel(
 	}
 }
 
-FLiveLinkHubClientId FLiveLinkHubProvider::UpdateAddressBookEntry(const FMessageAddress& ControlEndpointAddress, int32 DiscoveryProtocolVersion, ELiveLinkTopologyMode TopologyMode, const FString& Hostname)
+FLiveLinkHubClientId FLiveLinkHubProvider::UpdateAddressBookEntry(const FMessageAddress& ControlEndpointAddress, int32 DiscoveryProtocolVersion, ELiveLinkTopologyMode TopologyMode, const FString& Hostname, const FString& ProjectName)
 {
 	TOptional<FLiveLinkHubAddressBookEntry> ReadOnlyEntry;
 
@@ -1077,25 +1087,38 @@ FLiveLinkHubClientId FLiveLinkHubProvider::UpdateAddressBookEntry(const FMessage
 				}
 			}
 
+			if (Entry->ProjectName != ProjectName && !ProjectName.IsEmpty())
+			{
+				Entry->ProjectName = ProjectName;
+			}
+
 			ReadOnlyEntry = *Entry;
 		}
 	}
 
 	if (ReadOnlyEntry)
 	{
-		bool bModifiedIP = false;
+		bool bModifiedClientsMap = false;
 		{
 			UE::TWriteScopeLock Lock(ClientsMapLock);
 			// Update the clients info from the address book entry.
 			// Todo: Ideally, this map and the address book should be one and the same.
 			if (FLiveLinkHubUEClientInfo* Info = ClientsMap.Find(ReadOnlyEntry->ClientId); Info && Info->IPAddress != ReadOnlyEntry->IP)
 			{
-				Info->IPAddress = ReadOnlyEntry->IP;
-				bModifiedIP = true;
+				if (Info->IPAddress != ReadOnlyEntry->IP)
+				{
+					Info->IPAddress = ReadOnlyEntry->IP;
+					bModifiedClientsMap = true;
+				}
+				if (Info->ProjectName != ReadOnlyEntry->ProjectName && !ReadOnlyEntry->ProjectName.IsEmpty())
+				{
+					Info->ProjectName = ReadOnlyEntry->ProjectName;
+					bModifiedClientsMap = true;
+				}
 			}
 		}
 	
-		if (bModifiedIP)
+		if (bModifiedClientsMap)
 		{
 			// todo run this code again if auto connect behavior changes.
 			OnClientEventDelegate.Broadcast(ReadOnlyEntry->ClientId, EClientEventType::Modified);
@@ -1111,6 +1134,7 @@ FLiveLinkHubClientId FLiveLinkHubProvider::UpdateAddressBookEntry(const FMessage
 			.TopologyMode = TopologyMode,
 			.IP = LiveLinkHubProviderUtils::GetIPAddress(ControlEndpointAddress),
 			.Hostname = Hostname,
+			.ProjectName = ProjectName,
 			.LastBeaconTimestamp = FPlatformTime::Seconds(),
 			.DiscoveryProtocolVersion = DiscoveryProtocolVersion
 		};
@@ -1144,14 +1168,48 @@ void FLiveLinkHubProvider::HandleBeaconMessage(const FLiveLinkHubBeaconMessage&
 	{
 		// Beacons were added in the second version of the discovery protocol.
 		constexpr int32 DiscoveryProtocolVersion = 2;
-		UpdateAddressBookEntry(Context->GetSender(), DiscoveryProtocolVersion, Message.TopologyMode, Message.Hostname);
+		UpdateAddressBookEntry(Context->GetSender(), DiscoveryProtocolVersion, Message.TopologyMode, Message.Hostname, Message.ProjectName);
 	}
 }
 
-bool FLiveLinkHubProvider::ShouldAutoConnectTo(const FString& RemoteHostname) const
+bool FLiveLinkHubProvider::ShouldAutoConnectTo(const FString& RemoteHostname, const FString& IP, const FString& ProjectName) const
 {
+	FLiveLinkHubClientFilterPreset Preset = GetMutableDefault<ULiveLinkHubUserSettings>()->GetClientFiltersData_AnyThread();
+
+	bool bIncluded = false;
+
+	for (const FLiveLinkHubClientTextFilter& Filter : Preset.Filters)
+	{
+		const FString* StringToCompare = Filter.Type == ELiveLinkHubClientFilterType::IP ? &IP : &RemoteHostname;
+		// Allow filtering on project name if IP/Hostname is empty.
+		if (StringToCompare->MatchesWildcard(Filter.Text) || (Filter.Text.IsEmpty() && !Filter.Project.IsEmpty()))
+		{
+			const bool MatchesProject = Filter.Project.IsEmpty() || Filter.Project == TEXT("*") || ProjectName.MatchesWildcard(Filter.Project);
+			
+			if (MatchesProject)
+			{
+				if (Filter.Behavior == ELiveLinkHubClientFilterBehavior::Include)
+				{
+					// We don't return true here because it might get excluded in a subsequent filter.
+					bIncluded = true;
+				}
+				else
+				{
+					// Note: Exclusion currently has priority over inclusion filters.
+					return false;
+				}
+			}
+		}
+	}
+
+	if (bIncluded)
+	{
+		return true;
+	}
+
 	const bool bSameHost = RemoteHostname == GetMachineName();
-	const ELiveLinkHubAutoConnectMode AutoConnectClients = GetDefault<ULiveLinkHubSettings>()->AutoConnectClients;
+
+	const ELiveLinkHubAutoConnectMode AutoConnectClients = GetDefault<ULiveLinkHubUserSettings>()->GetAutoConnectMode();
 	if (AutoConnectClients != ELiveLinkHubAutoConnectMode::Disabled)
 	{
 		return AutoConnectClients == ELiveLinkHubAutoConnectMode::All || (AutoConnectClients == ELiveLinkHubAutoConnectMode::LocalOnly && bSameHost);
@@ -1184,4 +1242,37 @@ void FLiveLinkHubProvider::OnMessageBusNotification(const FMessageBusNotificatio
 	}
 }
 
+void FLiveLinkHubProvider::OnFiltersModified()
+{
+	// Remove all clients from session
+	TArray<FLiveLinkHubClientId> ClientIds;
+	{
+		UE::TReadScopeLock Lock(AddressBookLock);
+		ClientIds.Reserve(AddressBook.Num());
+		Algo::Transform(AddressBook, ClientIds, [](const FLiveLinkHubAddressBookEntry& Entry) { return Entry.ClientId; });
+	}
+
+	if (TSharedPtr<ILiveLinkHubSessionManager> Manager = SessionManager.Pin())
+	{
+		for (const FLiveLinkHubClientId& ClientId : ClientIds)
+		{
+			Manager->GetCurrentSession()->RemoveClient(ClientId);
+		}
+	}
+
+	FWriteScopeLock Lock(AddressBookLock);
+	for (auto EntryIt = AddressBook.CreateIterator(); EntryIt; ++EntryIt)
+	{
+		if (!EntryIt->bConnected && ShouldAutoConnectTo(EntryIt->Hostname, EntryIt->IP, EntryIt->ProjectName)
+			&& GetDefault<ULiveLinkHubMessagingSettings>()->CanTransmitTo(FLiveLinkHub::Get()->GetTopologyMode(), EntryIt->TopologyMode))
+		{
+			EntryIt->bSentDiscoveryRequest = true;
+			EntryIt->bDisableAutoconnect = false;
+
+			UE_LOG(LogLiveLinkHub, Verbose, TEXT("Sending Hub Connection request to %s"), *EntryIt->Hostname);
+			SendMessage(MakeLiveLinkHubDiscoveryMessage(), EntryIt->ControlEndpoint, EMessageFlags::Reliable);
+		}
+	}
+}
+
 #undef LOCTEXT_NAMESPACE /*LiveLinkHub.LiveLinkHubProvider*/
diff --git a/Engine/Plugins/Animation/LiveLinkHub/Source/LiveLinkHub/Private/Clients/LiveLinkHubProvider.h b/Engine/Plugins/Animation/LiveLinkHub/Source/LiveLinkHub/Private/Clients/LiveLinkHubProvider.h
--- a/Engine/Plugins/Animation/LiveLinkHub/Source/LiveLinkHub/Private/Clients/LiveLinkHubProvider.h
+++ b/Engine/Plugins/Animation/LiveLinkHub/Source/LiveLinkHub/Private/Clients/LiveLinkHubProvider.h
@@ -36,6 +36,8 @@ struct FLiveLinkHubAddressBookEntry
 	FString IP;
 	/** Hostname of the client. */
 	FString Hostname;
+	/** Project that the client is in. */
+	FString ProjectName;
 	/** Timestamp of the last beacon message received for this client. */
 	double LastBeaconTimestamp = 0.0;
 	/** This is set when a user removes a client from the session so as to not autoconnect it again. */
@@ -180,17 +182,26 @@ private:
 	void HandleBeaconMessage(const FLiveLinkHubBeaconMessage& Message, const TSharedRef<class IMessageContext, ESPMode::ThreadSafe>& Context);
 
 	/** Returns whether we should automatically connect to this client based on our settings. */
-	bool ShouldAutoConnectTo(const FString& RemoteHostname) const;
+	bool ShouldAutoConnectTo(const FString& RemoteHostname, const FString& IP = FString(), const FString& ProjectName = FString()) const;
 
 	/** 
 	 * Update (or creates if needed) an entry for this address.
 	 * @return the client ID.
 	 */
-	FLiveLinkHubClientId UpdateAddressBookEntry(const FMessageAddress& ControlEndpointAddress, int32 DiscoveryProtocolVersion, ELiveLinkTopologyMode TopologyMode = ELiveLinkTopologyMode::UnrealClient, const FString& Hostname = FString());
+	FLiveLinkHubClientId UpdateAddressBookEntry(const FMessageAddress& ControlEndpointAddress, int32 DiscoveryProtocolVersion, ELiveLinkTopologyMode TopologyMode = ELiveLinkTopologyMode::UnrealClient, const FString& Hostname = FString(), const FString& ProjectName = FString());
 
 	/** Utility method that constructs a discovery message. */
 	FLiveLinkHubDiscoveryMessage* MakeLiveLinkHubDiscoveryMessage() const;
 
+	/** Invoked when the user modifies the auto-connect filtering options. */
+	void OnFiltersModified();
+
+	/**
+	 * Connect to a set of clients.
+	 * @Note The clients must have been discovered before invoking this for the connection message to be sent.
+	 */
+	void ConnectToClients(const TArray<FLiveLinkHubClientId>& DiscoveredClientIds);
+
 protected:
 	//~ Begin ILiveLinkHubClientsModel interface
 	virtual void OnConnectionsClosed(const TArray<FMessageAddress>& ClosedAddresses) override;
diff --git a/Engine/Plugins/Animation/LiveLinkHub/Source/LiveLinkHub/Private/Clients/SLiveLinkHubClientFilters.h b/Engine/Plugins/Animation/LiveLinkHub/Source/LiveLinkHub/Private/Clients/SLiveLinkHubClientFilters.h
new file mode 100644
--- /dev/null
+++ b/Engine/Plugins/Animation/LiveLinkHub/Source/LiveLinkHub/Private/Clients/SLiveLinkHubClientFilters.h
@@ -0,0 +1,66 @@
+// Copyright Epic Games, Inc. All Rights Reserved.
+
+#pragma once
+
+#include "Widgets/SCompoundWidget.h"
+
+#include "Modules/ModuleManager.h"
+#include "PropertyEditorModule.h"
+#include "Settings/LiveLinkHubSettings.h"
+
+class SLiveLinkHubClientFilters : public SCompoundWidget
+{
+public:
+	SLATE_BEGIN_ARGS(SLiveLinkHubClientFilters) {}
+	SLATE_END_ARGS()
+
+	void Construct(const FArguments& InArgs)
+	{
+		ChildSlot
+		[
+			CreatePropertyEditor()
+		];
+	}
+
+	/** Create a property editor that displays the filters taken from ULiveLinkHubUserSettings. */
+	TSharedRef<SWidget> CreatePropertyEditor()
+	{
+		FPropertyEditorModule& PropertyEditorModule = FModuleManager::GetModuleChecked<FPropertyEditorModule>("PropertyEditor");
+
+		FDetailsViewArgs DetailsViewArgs;
+		DetailsViewArgs.bUpdatesFromSelection = false;
+		DetailsViewArgs.bLockable = false;
+		DetailsViewArgs.bShowPropertyMatrixButton = false;
+		DetailsViewArgs.NameAreaSettings = FDetailsViewArgs::HideNameArea;
+		DetailsViewArgs.ViewIdentifier = NAME_None;
+		DetailsViewArgs.bShowOptions = false;
+		DetailsViewArgs.bAllowSearch = false;
+		DetailsViewArgs.ColumnWidth = 0.8f;
+
+		TSharedRef<IDetailsView> DetailView = PropertyEditorModule.CreateDetailView( DetailsViewArgs );
+		DetailView->SetIsPropertyVisibleDelegate(FIsPropertyVisible::CreateSP(this, &SLiveLinkHubClientFilters::ShouldDisplayProperty));
+		DetailView->SetObject(GetMutableDefault<ULiveLinkHubUserSettings>());
+		DetailView->SetRootExpansionStates(true, true);
+
+		return DetailView;
+	}
+
+	/** Hides properties that aren't in the filters category. */
+	bool ShouldDisplayProperty(const FPropertyAndParent& InPropertyAndParent) const
+	{
+		static const FLazyName CategoryName = TEXT("Category");
+		static const FString FiltersCategory = TEXT("Filters");
+		if (InPropertyAndParent.Property.GetMetaData(CategoryName) == FiltersCategory)
+		{
+			return true;
+		}
+		if (!InPropertyAndParent.ParentProperties.IsEmpty())
+		{
+			if (InPropertyAndParent.ParentProperties.Last()->GetMetaData(CategoryName) == FiltersCategory)
+			{
+				return true;
+			}
+		}
+		return false;
+	}
+};
\ No newline at end of file
diff --git a/Engine/Plugins/Animation/LiveLinkHub/Source/LiveLinkHub/Private/Clients/SLiveLinkHubClientsView.h b/Engine/Plugins/Animation/LiveLinkHub/Source/LiveLinkHub/Private/Clients/SLiveLinkHubClientsView.h
--- a/Engine/Plugins/Animation/LiveLinkHub/Source/LiveLinkHub/Private/Clients/SLiveLinkHubClientsView.h
+++ b/Engine/Plugins/Animation/LiveLinkHub/Source/LiveLinkHub/Private/Clients/SLiveLinkHubClientsView.h
@@ -18,14 +18,18 @@
 #include "LiveLinkHubUEClientInfo.h"
 #include "Misc/Guid.h"
 #include "Misc/TextFilter.h"
+#include "Dialog/SCustomDialog.h"
 #include "Session/LiveLinkHubSessionManager.h"
+#include "Settings/LiveLinkHubSettings.h"
 #include "SLiveLinkFilterSearchBox.h"
+#include "SLiveLinkHubClientFilters.h"
 #include "SPositiveActionButton.h"
 #include "Styling/SlateStyleMacros.h"
 #include "Styling/SlateTypes.h"
 #include "Textures/SlateIcon.h"
 #include "Widgets/Input/SButton.h"
 #include "Widgets/Input/SCheckBox.h"
+#include "Widgets/Input/SEditableTextBox.h"
 #include "Widgets/Layout/SScaleBox.h"
 #include "Widgets/Layout/SWidgetSwitcher.h"
 #include "Widgets/Input/SComboButton.h"
@@ -394,7 +398,7 @@ public:
 					.ToolTipText_Lambda([]() 
 						{
 							const UEnum* Enum = StaticEnum<ELiveLinkHubAutoConnectMode>();
-							const int32 Index = Enum->GetIndexByValue(static_cast<int64>(GetDefault<ULiveLinkHubSettings>()->AutoConnectClients));
+							const int32 Index = Enum->GetIndexByValue(static_cast<int64>(GetDefault<ULiveLinkHubUserSettings>()->GetAutoConnectMode()));
 							return Enum->GetToolTipTextByIndex(Index);
 						})
 					.OnGetMenuContent_Raw(this, &SLiveLinkHubClientsView::GetAutoConnectMenuContent)
@@ -670,10 +674,13 @@ private:
 			DiscoveredClients.RemoveAt(ClientIndex);
 		}
 
-		TSharedPtr<FClientTreeViewClientItem> ClientItem = MakeShared<FClientTreeViewClientItem>(ClientId, ClientsModel.ToSharedRef());
-		ClientItem->ClientId = ClientId;
-		InitializeClientItem(*ClientItem);
-		Clients.Add(ClientItem);
+		if (!Clients.ContainsByPredicate([ClientId](const FClientTreeItemPtr& Item) { return Item && Item->ClientId == ClientId; }))
+		{
+			TSharedPtr<FClientTreeViewClientItem> ClientItem = MakeShared<FClientTreeViewClientItem>(ClientId, ClientsModel.ToSharedRef());
+			ClientItem->ClientId = ClientId;
+			InitializeClientItem(*ClientItem);
+			Clients.Add(ClientItem);
+		}
 
 		FilterSearchBox->Update();
 	}
@@ -693,6 +700,10 @@ private:
 		}
 
 		FilterSearchBox->Update();
+		if (TreeView)
+		{
+			TreeView->RequestListRefresh();
+		}
 	}
 
 	/** Refreshes the data when the current session changes. */
@@ -812,6 +823,41 @@ private:
 			}
 
 			MenuBuilder.EndSection();
+
+			MenuBuilder.BeginSection(NAME_None, LOCTEXT("FiltersLabel", "Filters"));
+
+			FString CurrentPresetName = GetDefault<ULiveLinkHubUserSettings>()->CurrentPreset.PresetName;
+
+			MenuBuilder.AddMenuEntry(LOCTEXT("FiltersLabel", "Filters"), LOCTEXT("FiltersToolTip", "Open the filters menu."), FSlateIcon(FAppStyle::GetAppStyleSetName(), "Icons.Filter"), FExecuteAction::CreateRaw(this, &SLiveLinkHubClientsView::OpenFiltersMenu));
+
+			FUIAction UIAction(
+				FExecuteAction::CreateLambda([]() {}),
+				FCanExecuteAction::CreateStatic([]()
+					{
+						return GetDefault<ULiveLinkHubUserSettings>()->GetFilterPresets().Num() > 0;
+					})
+			);
+
+			MenuBuilder.AddSubMenu(LOCTEXT("LoadPresetLabel", "Load Preset"), LOCTEXT("LoadPresetToolTip", "Load a set of filters from a saved preset."), FNewMenuDelegate::CreateRaw(this, &SLiveLinkHubClientsView::CreateLoadPresetSubMenu), UIAction, NAME_None, EUserInterfaceActionType::Button);
+
+			if (!CurrentPresetName.IsEmpty())
+			{
+				FUIAction CurrentPresetUIAction(
+					FExecuteAction::CreateLambda([]() {}),
+					FCanExecuteAction::CreateLambda([]()
+						{
+							return false;
+						})
+				);
+				MenuBuilder.AddMenuEntry(FText::Format(LOCTEXT("CurrentPresetLabel", "Current Preset: {0}"), FText::FromString(CurrentPresetName)), LOCTEXT("CurrentFilterPresetToolTip", "What filter preset is currently in use."), FSlateIcon(), CurrentPresetUIAction);
+				MenuBuilder.AddMenuEntry(LOCTEXT("SaveCurrentPresetLabel", "Save Current Preset"), LOCTEXT("SaveCurrentPreset", "Save the connected clients as filters to the current preset."), FSlateIcon(FAppStyle::GetAppStyleSetName(), "Icons.Save"), FExecuteAction::CreateRaw(this, &SLiveLinkHubClientsView::HandleSaveCurrentPreset));
+			}
+
+			MenuBuilder.AddMenuEntry(LOCTEXT("SaveNewPresetLabel", "Save Preset As"), LOCTEXT("SavePresetToolTip", "Save connected clients as a new preset."), FSlateIcon(FAppStyle::GetAppStyleSetName(), "Icons.Save"), FExecuteAction::CreateRaw(this, &SLiveLinkHubClientsView::HandleSavePreset));
+
+			MenuBuilder.AddSubMenu(LOCTEXT("DeletePresetLabel", "Delete Preset"), LOCTEXT("DeletePresetToolTip", "Delete a saved filter preset."), FNewMenuDelegate::CreateRaw(this, &SLiveLinkHubClientsView::CreateDeletePresetSubMenu), UIAction, NAME_None, EUserInterfaceActionType::Button);
+
+			MenuBuilder.EndSection();
 		}
 
 		return MenuBuilder.MakeWidget();
@@ -820,8 +866,7 @@ private:
 	/** Configure UE client AutoConnect mode. */
 	void SetAutoConnectMode(ELiveLinkHubAutoConnectMode AutoConnectMode)
 	{
-		GetMutableDefault<ULiveLinkHubSettings>()->AutoConnectClients = AutoConnectMode;
-		GetMutableDefault<ULiveLinkHubSettings>()->TryUpdateDefaultConfigFile();
+		GetMutableDefault<ULiveLinkHubUserSettings>()->SetAutoConnectMode(AutoConnectMode);
 	}
 
 	/** Create the widget for switching autoconnect mode. */
@@ -831,7 +876,7 @@ private:
 			SNew(SWidgetSwitcher)
 			.WidgetIndex_Lambda([]()
 			{
-				return StaticEnum<ELiveLinkHubAutoConnectMode>()->GetIndexByValue(static_cast<int64>(GetDefault<ULiveLinkHubSettings>()->AutoConnectClients));
+				return StaticEnum<ELiveLinkHubAutoConnectMode>()->GetIndexByValue(static_cast<int64>(GetDefault<ULiveLinkHubUserSettings>()->GetAutoConnectMode()));
 			});
 
 		const UEnum* Enum = StaticEnum<ELiveLinkHubAutoConnectMode>();
@@ -851,6 +896,121 @@ private:
 		return Switcher;
 	}
 
+	/** Open a window that hosts the filters menu. */
+	void OpenFiltersMenu()
+	{
+		SAssignNew(FiltersMenuWindow, SWindow)
+			.Title(LOCTEXT("EditAutoConnect", "Edit AutoConnect Filters"))
+			.ClientSize(FVector2D(600, 300))
+			.SizingRule(ESizingRule::UserSized)
+			[
+				SNew(SLiveLinkHubClientFilters)
+			];
+
+		if (GEditor)
+		{
+			GEditor->EditorAddModalWindow(FiltersMenuWindow.ToSharedRef());
+		}
+	}
+
+	/** Saves the current preset to disk. */
+	void HandleSaveCurrentPreset()
+	{
+		ULiveLinkHubUserSettings* Settings = GetMutableDefault<ULiveLinkHubUserSettings>();
+		Settings->SaveClientsToPreset(Settings->CurrentPreset.PresetName);
+	}
+
+	/** Open a window that hosts the filters menu. */
+	void HandleSavePreset()
+	{
+		FString PresetName;
+
+		TSharedPtr<SWidget> TextBoxWidget;
+		FText CurrentText;
+
+		TSharedPtr<SCustomDialog> SaveFilterDialog;
+
+		bool bConfirmedPresetName = false;
+
+		SaveFilterDialog = SNew(SCustomDialog)
+			.Title(FText(LOCTEXT("SaveFilters", "Save Filter Preset")))
+			.RootPadding(FMargin(2.0))
+			.ButtonAreaPadding(FMargin(2.0, 5.0))
+			.AutoCloseOnButtonPress(true)
+			.GetWidgetToFocusOnActivate(FGetFocusWidget::CreateLambda([&TextBoxWidget]() { return TextBoxWidget; }))
+			.Content()
+			[
+				SNew(SBox)
+				.MinDesiredWidth(150)
+				[
+					SNew(SVerticalBox)
+					+ SVerticalBox::Slot()
+					.HAlign(HAlign_Center)
+					[
+						SAssignNew(TextBoxWidget, SEditableTextBox)
+							.Font(FAppStyle::GetFontStyle("PropertyWindow.NormalFont"))
+							.Justification(ETextJustify::Left)
+							.MinDesiredWidth(60.f)
+							.SelectAllTextWhenFocused(true)
+							.OnTextChanged_Lambda([&CurrentText](const FText& NewText) { CurrentText = NewText; })
+							.HintText(LOCTEXT("PresetNameHint", "Preset Name"))
+							.OnTextCommitted_Lambda([&PresetName, &SaveFilterDialog, &bConfirmedPresetName](const FText& NewText, ETextCommit::Type InTextCommit)
+							{
+								PresetName = NewText.ToString();
+
+								// Make sure hitting escape is not considered 'confirming'.
+								bConfirmedPresetName |= InTextCommit == ETextCommit::OnEnter;
+
+								bool bDestroyWindow = InTextCommit == ETextCommit::OnCleared && !bConfirmedPresetName;
+
+								if (bConfirmedPresetName)
+								{
+									// Prevent destroying window if text is empty when confirming.
+									bDestroyWindow = !PresetName.IsEmpty();
+								}
+
+								if (bDestroyWindow)
+								{
+									SaveFilterDialog->RequestDestroyWindow();
+								}
+							})
+					]
+				]
+			]
+			.Buttons({
+				SCustomDialog::FButton(LOCTEXT("OK", "OK"), FSimpleDelegate(), SCustomDialog::EButtonRole::Confirm)
+					.SetIsEnabled(TAttribute<bool>::CreateLambda([&CurrentText]() { return !CurrentText.IsEmpty(); })),
+				SCustomDialog::FButton(LOCTEXT("Cancel", "Cancel"), FSimpleDelegate(), SCustomDialog::EButtonRole::Cancel)
+			})
+			.HAlignButtonBox(HAlign_Center);
+
+		
+
+		// returns 0 when OK is pressed, 1 when Cancel is pressed, -1 if the window is closed
+		const int NamePickerResult = SaveFilterDialog->ShowModal();
+		
+		if ((NamePickerResult == 0 || bConfirmedPresetName) && !PresetName.IsEmpty())
+		{
+			GetMutableDefault<ULiveLinkHubUserSettings>()->SaveClientsToPreset(PresetName);
+		}
+	}
+
+	void CreateLoadPresetSubMenu(FMenuBuilder& MenuBuilder)
+	{
+		for (const FLiveLinkHubClientFilterPreset& Preset : GetDefault<ULiveLinkHubUserSettings>()->GetFilterPresets())
+		{
+			MenuBuilder.AddMenuEntry(FText::FromString(Preset.PresetName), FText::GetEmpty(), FSlateIcon(FAppStyle::GetAppStyleSetName(), "Icons.Load"), FExecuteAction::CreateLambda([PresetName = Preset.PresetName]() { GetMutableDefault<ULiveLinkHubUserSettings>()->LoadFilterPreset(PresetName); }));
+		}
+	}
+
+	void CreateDeletePresetSubMenu(FMenuBuilder& MenuBuilder)
+	{
+		for (const FLiveLinkHubClientFilterPreset& Preset : GetDefault<ULiveLinkHubUserSettings>()->GetFilterPresets())
+		{
+			MenuBuilder.AddMenuEntry(FText::FromString(Preset.PresetName), FText::GetEmpty(), FSlateIcon(FAppStyle::GetAppStyleSetName(), "Icons.Load"), FExecuteAction::CreateLambda([PresetName = Preset.PresetName]() { GetMutableDefault<ULiveLinkHubUserSettings>()->DeleteFilterPreset(PresetName); }));
+		}
+	}
+
 private:
 	/** Delegate called when a client is selected. */
 	FOnClientSelected OnClientSelectedDelegate;
@@ -870,6 +1030,8 @@ private:
 	TArray<FClientTreeItemPtr> FilteredList;
 	/** Filter search box responsible for keeping the filtered list up to date. */
 	TSharedPtr<SLiveLinkFilterSearchBox<FClientTreeItemPtr>> FilterSearchBox;
+	/** Hosts the filter menu that allows better control over which sources should auto-connect. */
+	TSharedPtr<SWindow> FiltersMenuWindow;
 };
 
 #undef LOCTEXT_NAMESPACE /* LiveLinkHub.ClientsView */
diff --git a/Engine/Plugins/Animation/LiveLinkHub/Source/LiveLinkHub/Private/LiveLinkHubModule.cpp b/Engine/Plugins/Animation/LiveLinkHub/Source/LiveLinkHub/Private/LiveLinkHubModule.cpp
--- a/Engine/Plugins/Animation/LiveLinkHub/Source/LiveLinkHub/Private/LiveLinkHubModule.cpp
+++ b/Engine/Plugins/Animation/LiveLinkHub/Source/LiveLinkHub/Private/LiveLinkHubModule.cpp
@@ -16,6 +16,7 @@
 #include "Recording/LiveLinkHubPlaybackController.h"
 #include "Recording/LiveLinkHubRecordingController.h"
 #include "Settings/LiveLinkHubSettings.h"
+#include "Settings/LiveLinkHubClientTextFilterCustomization.h"
 #include "Settings/LiveLinkHubSettingsCustomization.h"
 #include "Settings/LiveLinkHubCustomTimeStepSettingsCustomization.h"
 #include "Settings/LiveLinkSettingsCustomization.h"
@@ -51,6 +52,7 @@ void FLiveLinkHubModule::StartupModule()
 	FPropertyEditorModule& PropertyModule = FModuleManager::LoadModuleChecked<FPropertyEditorModule>("PropertyEditor");
 
 	PropertyModule.RegisterCustomClassLayout(ULiveLinkHubSettings::StaticClass()->GetFName(), FOnGetDetailCustomizationInstance::CreateStatic(&FLiveLinkHubSettingsCustomization::MakeInstance));
+	PropertyModule.RegisterCustomPropertyTypeLayout(FLiveLinkHubClientTextFilter::StaticStruct()->GetFName(), FOnGetPropertyTypeCustomizationInstance::CreateStatic(&FLiveLinkHubClientTextFilterCustomization::MakeInstance));
 	PropertyModule.RegisterCustomPropertyTypeLayout(FLiveLinkHubCustomTimeStepSettings::StaticStruct()->GetFName(), FOnGetPropertyTypeCustomizationInstance::CreateStatic(&FLiveLinkHubCustomTimeStepSettingsCustomization::MakeInstance));
 
 	// Apply our customization for core live link settings, only if we aren't running in the full editor. We hide properties that aren't
@@ -67,6 +69,7 @@ void FLiveLinkHubModule::ShutdownModule()
 {
 	if (FPropertyEditorModule* PropertyEditorModule = FModuleManager::GetModulePtr<FPropertyEditorModule>("PropertyEditor"))
 	{
+		PropertyEditorModule->UnregisterCustomClassLayout(FLiveLinkHubClientTextFilter::StaticStruct()->GetFName());
 		PropertyEditorModule->UnregisterCustomClassLayout(ULiveLinkHubSettings::StaticClass()->GetFName());
 		if (bUseSettingsDetailCustomization)
 		{
diff --git a/Engine/Plugins/Animation/LiveLinkHub/Source/LiveLinkHub/Private/Settings/LiveLinkHubClientTextFilterCustomization.h b/Engine/Plugins/Animation/LiveLinkHub/Source/LiveLinkHub/Private/Settings/LiveLinkHubClientTextFilterCustomization.h
new file mode 100644
--- /dev/null
+++ b/Engine/Plugins/Animation/LiveLinkHub/Source/LiveLinkHub/Private/Settings/LiveLinkHubClientTextFilterCustomization.h
@@ -0,0 +1,98 @@
+// Copyright Epic Games, Inc. All Rights Reserved.
+
+#pragma once
+
+#include "IPropertyTypeCustomization.h"
+
+#include "CoreMinimal.h"
+#include "DetailLayoutBuilder.h"
+#include "DetailWidgetRow.h"
+#include "IDetailChildrenBuilder.h"
+#include "IDetailPropertyRow.h"
+#include "LiveLinkHubSettings.h"
+#include "PropertyHandle.h"
+#include "Widgets/Text/SInlineEditableTextBlock.h"
+
+/**
+ * Customization for the CustomTimeStep settings.
+ */
+class FLiveLinkHubClientTextFilterCustomization : public IPropertyTypeCustomization
+{
+public:
+	static TSharedRef<IPropertyTypeCustomization> MakeInstance()
+	{
+		return MakeShared<FLiveLinkHubClientTextFilterCustomization>();
+	}
+
+	virtual bool ShouldInlineKey() const override
+	{
+		return true;
+	}
+
+	virtual void CustomizeHeader(TSharedRef<IPropertyHandle> PropertyHandle, FDetailWidgetRow& HeaderRow, IPropertyTypeCustomizationUtils& CustomizationUtils)
+	{
+		if (!PropertyHandle->IsValidHandle())
+		{
+			return;
+		}
+
+		TSharedPtr<IPropertyUtilities> PropertyUtils = CustomizationUtils.GetPropertyUtilities();
+		TObjectPtr<class UScriptStruct> Struct = CastFieldChecked<FStructProperty>(PropertyHandle->GetProperty())->Struct;
+
+		check(Struct);
+		check(Struct == FLiveLinkHubClientTextFilter::StaticStruct());
+
+		PropertyHandle->SetOnChildPropertyValueChanged(FSimpleDelegate::CreateLambda([]() { GetMutableDefault<ULiveLinkHubUserSettings>()->PostUpdateClientFilters(); }));
+
+		TSharedPtr<IPropertyHandle> BehaviorHandle = PropertyHandle->GetChildHandle(GET_MEMBER_NAME_CHECKED(FLiveLinkHubClientTextFilter, Behavior));
+		TSharedPtr<IPropertyHandle> TypeHandle = PropertyHandle->GetChildHandle(GET_MEMBER_NAME_CHECKED(FLiveLinkHubClientTextFilter, Type));
+		TSharedPtr<IPropertyHandle> TextHandle = PropertyHandle->GetChildHandle(GET_MEMBER_NAME_CHECKED(FLiveLinkHubClientTextFilter, Text));
+		TSharedPtr<IPropertyHandle> ProjectHandle = PropertyHandle->GetChildHandle(GET_MEMBER_NAME_CHECKED(FLiveLinkHubClientTextFilter, Project));
+
+		HeaderRow.NameContent()
+			.MaxDesiredWidth(20) // use DPI?
+			[
+				PropertyHandle->CreatePropertyNameWidget()
+			]
+			.ValueContent()
+			[
+				SNew(SHorizontalBox)
+				+ SHorizontalBox::Slot()
+				.Padding(FMargin(2.0))
+				.VAlign(VAlign_Center)
+				.AutoWidth()
+				[
+					BehaviorHandle->CreatePropertyValueWidget()
+				]
+				+ SHorizontalBox::Slot()
+				.Padding(FMargin(2.0))
+				.VAlign(VAlign_Center)
+				.AutoWidth()
+				[
+					TypeHandle->CreatePropertyValueWidget()
+				]
+				+ SHorizontalBox::Slot()
+				.MinWidth(40.0)
+				[
+					TextHandle->CreatePropertyValueWidget()
+				]
+				+ SHorizontalBox::Slot()
+				.Padding(FMargin(4.0, 0.0))
+				.AutoWidth()
+				[
+					ProjectHandle->CreatePropertyNameWidget()
+				]
+				+ SHorizontalBox::Slot()
+				.Padding(FMargin(2.0, 0.0))
+				.MinWidth(80.f)
+				[
+					ProjectHandle->CreatePropertyValueWidget()
+				]
+			];
+	}
+
+	virtual void CustomizeChildren(TSharedRef<IPropertyHandle> PropertyHandle, IDetailChildrenBuilder& ChildBuilder, IPropertyTypeCustomizationUtils& CustomizationUtils)
+	{
+
+	}
+};
diff --git a/Engine/Plugins/Animation/LiveLinkHub/Source/LiveLinkHub/Private/Settings/LiveLinkHubSettings.h b/Engine/Plugins/Animation/LiveLinkHub/Source/LiveLinkHub/Private/Settings/LiveLinkHubSettings.h
--- a/Engine/Plugins/Animation/LiveLinkHub/Source/LiveLinkHub/Private/Settings/LiveLinkHubSettings.h
+++ b/Engine/Plugins/Animation/LiveLinkHub/Source/LiveLinkHub/Private/Settings/LiveLinkHubSettings.h
@@ -2,10 +2,15 @@
 
 #pragma once
 
+#include "Algo/Find.h"
+#include "Clients/LiveLinkHubUEClientInfo.h"
 #include "Config/LiveLinkHubTemplateTokens.h"
 #include "UObject/Object.h"
+#include "ILiveLinkHubClientsModel.h"
+#include "LiveLinkHubLog.h"
 #include "LiveLinkHubMessages.h"
 #include "Misc/FrameRate.h"
+#include "Misc/TransactionallySafeRWLock.h"
 #include "Templates/SubclassOf.h"
 
 #include "LiveLinkHubSettings.generated.h"
@@ -41,10 +46,6 @@ public:
 	UPROPERTY(config, EditAnywhere, Category = "LiveLinkHub")
 	FFilePath StartupConfig;
 
-	/** If enabled, discovered clients will be automatically added to the current session. */
-	UPROPERTY(config, EditAnywhere, Category="LiveLinkHub", DisplayName = "Auto Connect Mode")
-	ELiveLinkHubAutoConnectMode AutoConnectClients = ELiveLinkHubAutoConnectMode::LocalOnly;
-
 	/** The size in megabytes to buffer when streaming a recording. */
 	UPROPERTY(config, EditAnywhere, Category="LiveLinkHub", meta = (ClampMin = "1", UIMin = "1"))
 	int32 PlaybackFrameBufferSizeMB = 100;
@@ -138,6 +139,65 @@ private:
 	mutable TObjectPtr<ULiveLinkHubNamingTokens> NamingTokens;
 };
 
+/** Whether a filter should include or exlude the client that matches the filter. */
+UENUM()
+enum class ELiveLinkHubClientFilterBehavior
+{
+	Include,
+	Exclude
+};
+
+/** Whether we're filtering on IP or a hostname. */
+UENUM()
+enum class ELiveLinkHubClientFilterType
+{
+	IP,
+	Host
+};
+
+/** Basic text filter that includes or excludes IPs/Hostnames. */
+USTRUCT()
+struct FLiveLinkHubClientTextFilter
+{
+	GENERATED_BODY()
+
+	/** Optional project name filter. Will match any project if empty or *. */
+	UPROPERTY(config, EditAnywhere, Category = "Filters")
+	FString Project;
+
+	/** Text that will be used to match the filter. */
+	UPROPERTY(config, EditAnywhere, Category = "Filters")
+	FString Text;
+
+	/** Type of the filter. */
+	UPROPERTY(config, EditAnywhere, Category = "Filters")
+	ELiveLinkHubClientFilterType Type = ELiveLinkHubClientFilterType::IP;
+
+	/** Behavior of the filter. */
+	UPROPERTY(config, EditAnywhere, Category = "Filters")
+	ELiveLinkHubClientFilterBehavior Behavior = ELiveLinkHubClientFilterBehavior::Include;
+};
+
+/** A collection of filters, used to restore and save filters on disk. */
+USTRUCT()
+struct FLiveLinkHubClientFilterPreset
+{
+	GENERATED_BODY()
+
+	/** Name of this preset. */
+	UPROPERTY(config, DisplayName="Name")
+	FString PresetName;
+
+	/** General autoconnect mode of the app. */
+	UPROPERTY(config)
+	ELiveLinkHubAutoConnectMode AutoConnectClients = ELiveLinkHubAutoConnectMode::LocalOnly;
+
+	/** List of filters in the preset. */
+	UPROPERTY(config, EditAnywhere, Category = "Filters")
+	TArray<FLiveLinkHubClientTextFilter> Filters;
+};
+
+
 /**
  * User Settings for LiveLinkHub.
  */
@@ -148,6 +208,15 @@ class ULiveLinkHubUserSettings : public UObject
 
 public:
 
+	virtual void PostInitProperties() override
+	{
+		Super::PostInitProperties();
+
+		// Cache the filter data on load.
+		UE::TWriteScopeLock Lock{ FiltersLock };
+		CurrentFilterPreset_ThreadSafe = CurrentPreset;
+	}
+
 	/** Number of recent configs to keep track. */
 	UPROPERTY(config, EditAnywhere, Category = "LiveLinkHub")
 	uint8 MaxNumberOfRecentConfigs = 5;
@@ -164,12 +233,152 @@ public:
 	UPROPERTY(config)
 	FString LastConfigDirectory;
 
+	UPROPERTY(config, EditAnywhere, Category = "Filters", meta=(ShowOnlyInnerProperties))
+	FLiveLinkHubClientFilterPreset CurrentPreset;
+
 public:
 	/** Call after an explicit config save or load to save it in the recent config file list. */
 	void CacheRecentConfig(const FString& SavePath);
 
 	/** Get the most recent config / session files that were used in LiveLinkHub. */
 	TArray<struct FLiveLinkHubSessionFile> GetRecentConfigFiles();
+
+	/** Invoked through post-edit change, but can be invoked manually to copy the filters data to a thread-safe copy.*/
+	void PostUpdateClientFilters()
+	{
+		{
+			UE::TWriteScopeLock Lock{ FiltersLock };
+			CurrentFilterPreset_ThreadSafe = CurrentPreset;
+		}
+
+		OnFiltersModifiedDelegate.Broadcast();
+
+		SaveConfig();
+	}
+
+#if WITH_EDITOR
+	virtual void PostEditChangeProperty(struct FPropertyChangedEvent& PropertyChangedEvent)
+	{
+		PostUpdateClientFilters();
+	}
+#endif
+
+	/** Get the AutoConnect mode from the current preset. */
+	ELiveLinkHubAutoConnectMode GetAutoConnectMode() const 
+	{
+		UE::TReadScopeLock Lock{ FiltersLock };
+		return CurrentFilterPreset_ThreadSafe.AutoConnectClients;
+	}
+
+	/** Sets the autoconnect mode. */
+	void SetAutoConnectMode(ELiveLinkHubAutoConnectMode Mode)
+	{
+		CurrentPreset.AutoConnectClients = Mode;
+
+		PostUpdateClientFilters();
+	}
+
+	/** Creates a preset of filters from the list of clients that are currently in the session. */
+	void SaveClientsToPreset(const FString& InPresetName)
+	{
+		TArray<TPair<FString, FString>> IPAndProject;
+		if (ILiveLinkHubClientsModel* ClientsModel = ILiveLinkHubClientsModel::Get())
+		{
+			for (const FLiveLinkHubClientId& ClientId : ClientsModel->GetSessionClients())
+			{
+				if (TOptional<FLiveLinkHubUEClientInfo> Info = ClientsModel->GetClientInfo(ClientId))
+				{
+					if (!Info->IPAddress.IsEmpty())
+					{
+						IPAndProject.Emplace(Info->IPAddress, Info->ProjectName);
+					}
+				}
+			}
+		}
+
+		FLiveLinkHubClientFilterPreset Preset;
+		Preset.PresetName = InPresetName;
+		Preset.Filters.Reserve(IPAndProject.Num());
+		Preset.AutoConnectClients = CurrentPreset.AutoConnectClients;
+
+		Algo::Transform(IPAndProject, Preset.Filters, [](const TPair<FString, FString>& IpAndProject)
+		{
+			return FLiveLinkHubClientTextFilter{ IpAndProject.Value, IpAndProject.Key, ELiveLinkHubClientFilterType::IP, ELiveLinkHubClientFilterBehavior::Include };
+		});
+
+		int32 Index = FilterPresets.IndexOfByPredicate([&InPresetName](const FLiveLinkHubClientFilterPreset& Preset) { return InPresetName == Preset.PresetName; });
+		if (Index != INDEX_NONE)
+		{
+			FilterPresets[Index] = Preset;
+		}
+		else
+		{
+			FilterPresets.Add(Preset);
+		}
+
+		CurrentPreset = Preset;
+
+		SaveConfig();
+
+		UE::TWriteScopeLock Lock{ FiltersLock };
+		CurrentFilterPreset_ThreadSafe = CurrentPreset;
+	}
+
+	/** Get all filter presets. */
+	const TArray<FLiveLinkHubClientFilterPreset>& GetFilterPresets() const
+	{
+		return FilterPresets;
+	}
+
+	/** Applies a filter preset. */
+	void LoadFilterPreset(const FString& PresetName)
+	{
+		if (FLiveLinkHubClientFilterPreset* Preset = Algo::FindBy(FilterPresets, PresetName, &FLiveLinkHubClientFilterPreset::PresetName))
+		{
+			{
+				UE::TWriteScopeLock Lock{ FiltersLock };
+				CurrentPreset = *Preset;
+			}
+
+			PostUpdateClientFilters();
+		}
+		else
+		{
+			UE_LOG(LogLiveLinkHub, Warning, TEXT("Could not find preset %s."), *PresetName);
+		}
+	}
+
+	/** Deletes a filter preset. */
+	void DeleteFilterPreset(const FString& PresetName)
+	{
+		FilterPresets.RemoveAll([&PresetName](const FLiveLinkHubClientFilterPreset& InPreset) { return PresetName == InPreset.PresetName; });
+		SaveConfig();
+	}
+
+	/** Get a copy of the client filters data in a thread-safe way. */
+	FLiveLinkHubClientFilterPreset GetClientFiltersData_AnyThread()
+	{
+		UE::TReadScopeLock Lock{ FiltersLock };
+		return CurrentFilterPreset_ThreadSafe;
+	}
+
+	DECLARE_TS_MULTICAST_DELEGATE(FOnFiltersModified);
+	FOnFiltersModified OnFiltersModifiedDelegate;
+
+private:
+	/** 
+	 * Shadow variable needed to access the current filter preset from outside the game thread. 
+	 * Used to avoid excessive locking from widgets ticking on the game thread and accessing CurrentPreset. 
+	 */
+	FLiveLinkHubClientFilterPreset CurrentFilterPreset_ThreadSafe;
+
+	/** List of filter presets serialized to the user config. */
+	UPROPERTY(config)
+	TArray<FLiveLinkHubClientFilterPreset> FilterPresets;
+
+	/** Lock used to access CurrentFilterPreset_ThreadSafe. */
+	mutable FTransactionallySafeRWLock FiltersLock;
 };
 
 #undef UE_API
+
diff --git a/Engine/Plugins/Animation/LiveLinkHub/Source/LiveLinkHubMessaging/Private/LiveLinkHubControlChannel.h b/Engine/Plugins/Animation/LiveLinkHub/Source/LiveLinkHubMessaging/Private/LiveLinkHubControlChannel.h
--- a/Engine/Plugins/Animation/LiveLinkHub/Source/LiveLinkHubMessaging/Private/LiveLinkHubControlChannel.h
+++ b/Engine/Plugins/Animation/LiveLinkHub/Source/LiveLinkHubMessaging/Private/LiveLinkHubControlChannel.h
@@ -141,7 +141,7 @@ public:
 		FLiveLinkHubMessagingModule& Module = FModuleManager::Get().GetModuleChecked<FLiveLinkHubMessagingModule>("LiveLinkHubMessaging");
 		ELiveLinkTopologyMode Mode = Module.GetHostTopologyMode();
 
-		Endpoint->Send(FMessageEndpoint::MakeMessage<FLiveLinkHubBeaconMessage>(Mode, Hostname), EMessageFlags::Reliable, {}, nullptr, { RemoteAddress }, FTimespan::Zero(), FDateTime::MaxValue());
+		Endpoint->Send(FMessageEndpoint::MakeMessage<FLiveLinkHubBeaconMessage>(Mode, Hostname, FApp::GetProjectName()), EMessageFlags::Reliable, {}, nullptr, { RemoteAddress }, FTimespan::Zero(), FDateTime::MaxValue());
 	}
 
 	/** 
diff --git a/Engine/Plugins/Animation/LiveLinkHub/Source/LiveLinkHubMessaging/Public/LiveLinkHubMessages.h b/Engine/Plugins/Animation/LiveLinkHub/Source/LiveLinkHubMessaging/Public/LiveLinkHubMessages.h
--- a/Engine/Plugins/Animation/LiveLinkHub/Source/LiveLinkHubMessaging/Public/LiveLinkHubMessages.h
+++ b/Engine/Plugins/Animation/LiveLinkHub/Source/LiveLinkHubMessaging/Public/LiveLinkHubMessages.h
@@ -216,6 +216,10 @@ struct FLiveLinkHubBeaconMessage
 	UPROPERTY()
 	FString Hostname = FPlatformProcess::ComputerName();
 
+	/** Name of the machine that sent the beacon. */
+	UPROPERTY()
+	FString ProjectName;
+
 	/** Unique ID for a given client process. */
 	UPROPERTY()
 	FGuid InstanceId = FApp::GetInstanceId();
