LiveLinkHub: Fix freeze when looping large recordings.
- Fix freeze if waiting for frames to buffer where the range required was spread across different tracks, but not within a single track.
- Fix a possible freeze when reversing or stopping a recording immediately after looping.

diff --git a/Engine/Plugins/Animation/LiveLinkHub/Source/LiveLinkHub/Private/Recording/Implementations/LiveLinkRecordingDataContainer.cpp b/Engine/Plugins/Animation/LiveLinkHub/Source/LiveLinkHub/Private/Recording/Implementations/LiveLinkRecordingDataContainer.cpp
--- a/Engine/Plugins/Animation/LiveLinkHub/Source/LiveLinkHub/Private/Recording/Implementations/LiveLinkRecordingDataContainer.cpp
+++ b/Engine/Plugins/Animation/LiveLinkHub/Source/LiveLinkHub/Private/Recording/Implementations/LiveLinkRecordingDataContainer.cpp
@@ -27,8 +27,14 @@ TRange<int32> FLiveLinkRecordingBaseDataContainer::GetBufferedFrames(bool bInclu
 	}
 	int32 StartFrame = RecordedDataStartFrame;
 	int32 EndFrame = FMath::Max(0, RecordedDataStartFrame + Timestamps.Num() - 1);
-	if (bIncludeOffset)
+	if (bIncludeOffset && LocalFrameOffset > 0)
 	{
+		// We only add on positive frame offsets because negative ones, while accurate, don't correctly report this track's frames
+		// at a global level. So if we're waiting for high frames to buffer that don't really exist because they actually occurred in the past, we
+		// will freeze indefinitely. Such as we need frame 5000, but only one track has this frame at a negative offset because it ran at a
+		// faster framerate, so it's already happened and converts to a lower frame.
+		// @todo: We should revamp how loaded frames are reported, or how we wait for a frame range to buffer.
+		
 		StartFrame += LocalFrameOffset;
 		EndFrame += LocalFrameOffset;
 	}
diff --git a/Engine/Plugins/Animation/LiveLinkHub/Source/LiveLinkHub/Private/Recording/Implementations/LiveLinkUAssetRecording.cpp b/Engine/Plugins/Animation/LiveLinkHub/Source/LiveLinkHub/Private/Recording/Implementations/LiveLinkUAssetRecording.cpp
--- a/Engine/Plugins/Animation/LiveLinkHub/Source/LiveLinkHub/Private/Recording/Implementations/LiveLinkUAssetRecording.cpp
+++ b/Engine/Plugins/Animation/LiveLinkHub/Source/LiveLinkHub/Private/Recording/Implementations/LiveLinkUAssetRecording.cpp
@@ -332,15 +332,18 @@ UE::LiveLinkHub::RangeHelpers::Private::TRangeArray<int32> ULiveLinkUAssetRecord
 
 bool ULiveLinkUAssetRecording::IsFrameRangeBuffered(const TRange<int32>& InRange) const
 {
+	TRangeSet<int32> Union;
+	
 	const UE::LiveLinkHub::RangeHelpers::Private::TRangeArray<int32> Ranges = GetBufferedFrameRanges();
 	for (const TRange<int32>& Range : Ranges)
 	{
-		if (Range.Contains(InRange))
+		if (!Range.IsEmpty())
 		{
-			return true;
+			Union.Add(Range);
 		}
 	}
-	return false;
+	
+	return Union.Contains(InRange);
 }
 
 void ULiveLinkUAssetRecording::CopyRecordingData(FLiveLinkPlaybackTracks& InOutLiveLinkPlaybackTracks) const
diff --git a/Engine/Plugins/Animation/LiveLinkHub/Source/LiveLinkHub/Private/Recording/LiveLinkHubPlaybackController.cpp b/Engine/Plugins/Animation/LiveLinkHub/Source/LiveLinkHub/Private/Recording/LiveLinkHubPlaybackController.cpp
--- a/Engine/Plugins/Animation/LiveLinkHub/Source/LiveLinkHub/Private/Recording/LiveLinkHubPlaybackController.cpp
+++ b/Engine/Plugins/Animation/LiveLinkHub/Source/LiveLinkHub/Private/Recording/LiveLinkHubPlaybackController.cpp
@@ -400,6 +400,11 @@ void FLiveLinkHubPlaybackController::BeginPlayback(bool bInReverse)
 
 void FLiveLinkHubPlaybackController::RestartPlayback()
 {
+	if (bIsStopping)
+	{
+		// A stop request is in progress. It's possible a restart request also is occurring because of looping.
+		return;
+	}
 	const bool bOldReverse = bIsReverse; // Stop playback resets reverse
 	StopPlayback();
 	StartTimestamp = GetCurrentTime().AsSeconds();
@@ -418,11 +423,12 @@ void FLiveLinkHubPlaybackController::PausePlayback()
 
 void FLiveLinkHubPlaybackController::StopPlayback()
 {
+	bIsStopping = true;
 	bIsPlaying = false;
-
+	
 	// Wait for the playback thread to exit...
 	const uint32 CurrentThreadId = FPlatformTLS::GetCurrentThreadId();
-	if (CurrentThreadId != Thread->GetThreadID())
+	if (CurrentThreadId != Thread->GetThreadID() && IsInGameThread() /* Only wait if request came from gamethread, not the playback thread */)
 	{
 		while(!bIsPlaybackWaiting) { }
 	}
@@ -434,6 +440,7 @@ void FLiveLinkHubPlaybackController::StopPlayback()
 
 	RecordingPlayer->RestartPlayback();
 	bIsReverse = false;
+	bIsStopping = false;
 }
 
 void FLiveLinkHubPlaybackController::Eject(TFunction<void()> CompletionCallback)
diff --git a/Engine/Plugins/Animation/LiveLinkHub/Source/LiveLinkHub/Private/Recording/LiveLinkHubPlaybackController.h b/Engine/Plugins/Animation/LiveLinkHub/Source/LiveLinkHub/Private/Recording/LiveLinkHubPlaybackController.h
--- a/Engine/Plugins/Animation/LiveLinkHub/Source/LiveLinkHub/Private/Recording/LiveLinkHubPlaybackController.h
+++ b/Engine/Plugins/Animation/LiveLinkHub/Source/LiveLinkHub/Private/Recording/LiveLinkHubPlaybackController.h
@@ -261,6 +261,8 @@ private:
 	std::atomic<bool> bIsPlaybackWaiting = false;
 	/** Whether a recording is playing. */
 	std::atomic<bool> bIsPlaying = false;
+	/** Stop requested. */
+	std::atomic<bool> bIsStopping = false;
 	/** Whether we're currently paused. */
 	std::atomic<bool> bIsPaused = false;
 	/** If the recording is playing in reverse. */
