LiveLinkHub - Fix wrong invocation of frame received delegate

diff --git a/Engine/Plugins/Animation/LiveLinkHub/Source/LiveLinkHub/Private/LiveLinkHub.cpp b/Engine/Plugins/Animation/LiveLinkHub/Source/LiveLinkHub/Private/LiveLinkHub.cpp
--- a/Engine/Plugins/Animation/LiveLinkHub/Source/LiveLinkHub/Private/LiveLinkHub.cpp
+++ b/Engine/Plugins/Animation/LiveLinkHub/Source/LiveLinkHub/Private/LiveLinkHub.cpp
@@ -438,7 +438,7 @@ void FLiveLinkHub::OnLiveLinkFrameDataReceived(const FLiveLinkFrameDataStruct&)
 {
 	if (CVarLiveLinkHubEnableTriggeredTick.GetValueOnAnyThread())
 	{
-		OnFrameReceivedDelegate.Execute();
+		OnFrameReceivedDelegate.ExecuteIfBound();
 	}
 }
 
