LiveLinkHub: Fix issues when jumping to a specific frame.
Waiting for a frame to buffer will now properly validate the frame is ready across all tracks.

Fixes several issues:
- Previously buffered frames that were cached and pending removal would still be considered buffered even if they weren't ready, resulting in the wrong frame being displayed initially.
- Jumping to a frame that was buffered and ready on one track, but not ready on another track that had inconsistent frames or a faster frame rate would not wait for this track's frame to be ready, displaying the wrong frame.
- Reversing from the start to end frame, such as from a reverse play loop, would hit one or both of these issues and the animation would not always continue playing correctly.

diff --git a/Engine/Plugins/Animation/LiveLinkHub/Source/LiveLinkHub/Private/Recording/Implementations/LiveLinkUAssetRecording.cpp b/Engine/Plugins/Animation/LiveLinkHub/Source/LiveLinkHub/Private/Recording/Implementations/LiveLinkUAssetRecording.cpp
--- a/Engine/Plugins/Animation/LiveLinkHub/Source/LiveLinkHub/Private/Recording/Implementations/LiveLinkUAssetRecording.cpp
+++ b/Engine/Plugins/Animation/LiveLinkHub/Source/LiveLinkHub/Private/Recording/Implementations/LiveLinkUAssetRecording.cpp
@@ -317,7 +317,7 @@ bool ULiveLinkUAssetRecording::WaitForBufferedFrames(int32 InMinFrame, int32 InM
 	return true;
 }
 
-UE::LiveLinkHub::RangeHelpers::Private::TRangeArray<int32> ULiveLinkUAssetRecording::GetBufferedFrameRanges() const
+UE::LiveLinkHub::RangeHelpers::Private::TRangeArray<int32> ULiveLinkUAssetRecording::GetBufferedFrameRanges(bool bIncludeInactive) const
 {
 	DECLARE_SCOPE_CYCLE_COUNTER(TEXT("ULiveLinkUAssetRecording::GetBufferedFrameRanges"), STAT_ULiveLinkUAssetRecording_GetBufferedFrameRanges, STATGROUP_LiveLinkHub);
 	
@@ -327,14 +327,21 @@ UE::LiveLinkHub::RangeHelpers::Private::TRangeArray<int32> ULiveLinkUAssetRecord
 	}
 
 	FScopeLock Lock(&BufferedFrameMutex);
-	return BufferedFrameRanges;
+	UE::LiveLinkHub::RangeHelpers::Private::TRangeArray<int32> Result = BufferedFrameRanges;
+	if (bIncludeInactive)
+	{
+		Result.Append(InactiveBufferedFrameRanges);
+	}
+	return Result;
 }
 
 bool ULiveLinkUAssetRecording::IsFrameRangeBuffered(const TRange<int32>& InRange) const
 {
+	DECLARE_SCOPE_CYCLE_COUNTER(TEXT("ULiveLinkUAssetRecording::IsFrameRangeBuffered"), STAT_ULiveLinkUAssetRecording_IsFrameRangeBuffered, STATGROUP_LiveLinkHub);
+    	
 	TRangeSet<int32> Union;
 	
-	const UE::LiveLinkHub::RangeHelpers::Private::TRangeArray<int32> Ranges = GetBufferedFrameRanges();
+	const UE::LiveLinkHub::RangeHelpers::Private::TRangeArray<int32> Ranges = GetBufferedFrameRanges(/*bIncludeInactive*/ false);
 	for (const TRange<int32>& Range : Ranges)
 	{
 		if (!Range.IsEmpty())
@@ -342,8 +349,59 @@ bool ULiveLinkUAssetRecording::IsFrameRangeBuffered(const TRange<int32>& InRange
 			Union.Add(Range);
 		}
 	}
+
+	// Quick verification the buffered range contains our desired range. 
+	if (Union.Contains(InRange))
+	{
+		// -------------------------------------------------------------------------------------------------------------------------------------
+		// The buffered range is only exact across all tracks the initial time frames are loaded, since the range is only updated once a given
+		// frame is buffered for all tracks. Once we move the scrubber position, it's no longer completely accurate, as one track might have some
+		// frames unloaded, and another track might have moved frames to a cache. This depends on buffer size, track frame rate, etc.
+
+		// We can run into issues when one track has frames loaded for the range size, but another track no longer does. Without additional checks
+		// this method won't correctly report if the range is really buffered, causing the wrong frame to be displayed for certain tracks.
+
+		// Rather than check if every frame in the range is truly loaded across all tracks, we can just check if the target frame is. This
+		// method is only used to wait for loaded frames and gets updated when the scrubber position changes, so this single check should
+		// always ensure the frame we are displaying across all tracks is correct while being reasonably efficient.
+		// -------------------------------------------------------------------------------------------------------------------------------------
+		
+		// Check the target frame we're landing on is actually ready. The lower bound value is the exact
+		// frame we are scrubbing to, so it must be ready on all tracks.
+		const FQualifiedFrameTime TargetFrame(FFrameTime(InRange.GetLowerBoundValue()), GetGlobalFrameRate());
+		const double TargetFrameTime = TargetFrame.AsSeconds();
+		const int32 FrameIdx = TargetFrame.Time.GetFrame().Value;
+		
+		using namespace UE::LiveLinkHub::FrameData::Private;
+
+		// Iterate each frame data and verify the target frame is loaded and ready.
+		for (const TPair<FLiveLinkSubjectKey, FLiveLinkRecordingBaseDataContainer>& Pair : RecordingData.FrameData)
+		{
+			// Quickly check the frame idx to see if it is loaded for the track. O(1)
+			if (!Pair.Value.IsFrameLoaded(FrameIdx))
+			{
+				// If the frame is not loaded, but we've passed the buffer check above, then there is at least one other track
+				// that has this frame loaded. What might be happening is this particular track has different frame dispersion.
+
+				// O(1)
+				const FFrameMetaData& MetaData = FrameFileData.FindChecked(Pair.Key);
+
+				// Find the closest frame to the target frame: O(log n)
+				// This frame is either approximately close, or it could be farther in the past or future depending on the track's frames.
+				// Either way, this frame is required and will be loaded next if it isn't already.
+				const int32 LocalTargetFrameHint = TargetFrame.ConvertTo(MetaData.LocalFrameRate).FrameNumber.Value;
+				const int32 ClosestFrameIdx = MetaData.FindFrameFromTimestamp(TargetFrameTime, LocalTargetFrameHint);
+				if (!Pair.Value.IsFrameLoaded(ClosestFrameIdx))
+				{
+					return false;
+				}
+			}
+		}
+	
+		return true;
+	}
 	
-	return Union.Contains(InRange);
+	return false;
 }
 
 void ULiveLinkUAssetRecording::CopyRecordingData(FLiveLinkPlaybackTracks& InOutLiveLinkPlaybackTracks) const
@@ -888,6 +946,9 @@ void ULiveLinkUAssetRecording::LoadFrameData(UE::LiveLinkHub::FrameData::Private
 		{
 			// Don't go past requested frames or max frames.
 			MaxFrames = FMath::Min(MaxFrames, RequestedStartFrame + RequestedFramesToLoad + 1);
+			// Ensure we're always at least streaming in the initial frame. This is usually the case, but depending on the track's
+			// frames and how the RequestedFramesToLoad was divided, it may cut off the initial frame if it's the last frame of the track.
+			MaxFrames = FMath::Max(MaxFrames, RequestedInitialFrame + 1);
 		}
 
 		const int32 LastFrame = MaxFrames - 1;
@@ -1472,6 +1533,7 @@ void ULiveLinkUAssetRecording::UpdateBufferedFrames()
 	// Find all buffered ranges and convert them to the global frame rate. We don't want to report the local frames that are buffered, but
 	// the ranges as viewed by the recording/scrubber.
 	BufferedFrameRanges.Reset();
+	InactiveBufferedFrameRanges.Reset();
 	for (TTuple<FLiveLinkSubjectKey, FrameData::Private::FFrameMetaData>& FrameDataKeyVal : FrameFileData)
 	{
 		FrameData::Private::FFrameMetaData& FrameData = FrameDataKeyVal.Value;
@@ -1484,7 +1546,7 @@ void ULiveLinkUAssetRecording::UpdateBufferedFrames()
 		TRange<int32> ActiveBufferFrames = RecordingData.FrameData.FindChecked(*FrameData.FrameDataSubjectKey.Get()).GetBufferedFrames(/*bIncludeOffset*/ true);
 		ActiveBufferFrames = RangeHelpers::Private::ConvertRangeFrameRate(ActiveBufferFrames, FrameData.LocalFrameRate, GlobalFrameRate);
 		
-		BufferedFrameRanges.Append(MoveTemp(CacheBufferRanges));
+		InactiveBufferedFrameRanges.Append(MoveTemp(CacheBufferRanges));
 		BufferedFrameRanges.Add(MoveTemp(ActiveBufferFrames));
 	}
 }
diff --git a/Engine/Plugins/Animation/LiveLinkHub/Source/LiveLinkHub/Private/Recording/Implementations/LiveLinkUAssetRecording.h b/Engine/Plugins/Animation/LiveLinkHub/Source/LiveLinkHub/Private/Recording/Implementations/LiveLinkUAssetRecording.h
--- a/Engine/Plugins/Animation/LiveLinkHub/Source/LiveLinkHub/Private/Recording/Implementations/LiveLinkUAssetRecording.h
+++ b/Engine/Plugins/Animation/LiveLinkHub/Source/LiveLinkHub/Private/Recording/Implementations/LiveLinkUAssetRecording.h
@@ -60,8 +60,11 @@ public:
 	/** The size in bytes of each animation frame. */
 	int32 GetFrameDiskSize() const { return MaxFrameDiskSize; }
 	
-	/** Return all buffered frame ranges, including inactive. */
-	UE::LiveLinkHub::RangeHelpers::Private::TRangeArray<int32> GetBufferedFrameRanges() const;
+	/**
+	 * Return all buffered frame ranges.
+	 * @param bIncludeInactive If buffer ranges that are considered inactive should be included in the result.
+	 */
+	UE::LiveLinkHub::RangeHelpers::Private::TRangeArray<int32> GetBufferedFrameRanges(bool bIncludeInactive = true) const;
 	
 	/** Checks if a specific frame range is buffered. */
 	bool IsFrameRangeBuffered(const TRange<int32>& InRange) const;
@@ -212,6 +215,8 @@ private:
 
 	/** Frames buffered, divided into ranges. */
 	UE::LiveLinkHub::RangeHelpers::Private::TRangeArray<int32> BufferedFrameRanges;
+	/** Ranges that are buffered but considered inactive. */
+	UE::LiveLinkHub::RangeHelpers::Private::TRangeArray<int32> InactiveBufferedFrameRanges;
 	
 	/** The first (left most) frame to stream. */
 	int32 EarliestFrameToStream = 0;
