LiveLinkHub - Fix tooltip in LiveLinkHubMessagingSettings

#rnx

diff --git a/Engine/Plugins/Animation/LiveLinkHub/Source/LiveLinkHubMessaging/Public/LiveLinkHubMessagingSettings.h b/Engine/Plugins/Animation/LiveLinkHub/Source/LiveLinkHubMessaging/Public/LiveLinkHubMessagingSettings.h
--- a/Engine/Plugins/Animation/LiveLinkHub/Source/LiveLinkHubMessaging/Public/LiveLinkHubMessagingSettings.h
+++ b/Engine/Plugins/Animation/LiveLinkHub/Source/LiveLinkHubMessaging/Public/LiveLinkHubMessagingSettings.h
@@ -39,7 +39,7 @@ public:
 		};
 	}
 
-	/** Whether to allow the hub to receive LiveLink data from an unreal instance. (Useful if Unreal is using Live Link Broadcast components. */
+	/** Whether to allow the hub to receive LiveLink data from an unreal instance. (Useful if Unreal is using Live Link Broadcast components) */
 	UPROPERTY(config, EditAnywhere, Category = "LiveLinkHub")
 	bool bAllowReceivingFromUnreal = false;
 
