LiveLinkHub - Fix not being able to read config file after it was modified in VSCode

This was caused by VS Code saving the file in Little Endian BOM format rather than Little Endian, which wasn't getting read correctly by Unreal's json parser. This change works around this by loading the file to a String then passing it on to the parser.

diff --git a/Engine/Plugins/Animation/LiveLinkHub/Source/LiveLinkHub/Private/Config/LiveLinkHubFileUtilities.cpp b/Engine/Plugins/Animation/LiveLinkHub/Source/LiveLinkHub/Private/Config/LiveLinkHubFileUtilities.cpp
--- a/Engine/Plugins/Animation/LiveLinkHub/Source/LiveLinkHub/Private/Config/LiveLinkHubFileUtilities.cpp
+++ b/Engine/Plugins/Animation/LiveLinkHub/Source/LiveLinkHub/Private/Config/LiveLinkHubFileUtilities.cpp
@@ -7,6 +7,7 @@
 #include "JsonObjectConverter.h"
 #include "LiveLinkHub.h"
 #include "LiveLinkHubLog.h"
+#include "Misc/FileHelper.h"
 #include "Session/LiveLinkHubSessionData.h"
 #include "Session/LiveLinkHubSessionManager.h"
 #include "UObject/Package.h"
@@ -15,16 +16,13 @@ void UE::LiveLinkHub::FileUtilities::Private::SaveConfig(const ULiveLinkHubSessi
 {
 	if (ensure(!InFilePath.IsEmpty()))
 	{
-		const TUniquePtr<FArchive> Ar(IFileManager::Get().CreateFileWriter(*InFilePath));
-		if (Ar)
-		{
-			const TSharedPtr<FJsonObject> JsonObject = ToJson(InConfigData);
-			
-			const TSharedRef<TJsonWriter<>> JsonWriter = TJsonWriterFactory<>::Create(Ar.Get(), 0);
-			FJsonSerializer::Serialize(JsonObject.ToSharedRef(), JsonWriter);
+		FString JsonString;
 
-			ensure(Ar->Close());
-		}
+		const TSharedPtr<FJsonObject> JsonObject = ToJson(InConfigData);
+		const TSharedRef<TJsonWriter<>> JsonWriter = TJsonWriterFactory<>::Create(&JsonString, 0);
+		FJsonSerializer::Serialize(JsonObject.ToSharedRef(), JsonWriter);
+
+		FFileHelper::SaveStringToFile(JsonString, *InFilePath, FFileHelper::EEncodingOptions::ForceUTF8);
 	}
 }
 
@@ -32,24 +30,38 @@ ULiveLinkHubSessionData* UE::LiveLinkHub::FileUtilities::Private::LoadConfig(con
 {
 	if (IFileManager::Get().FileExists(*InFilePath))
 	{
-		const TUniquePtr<FArchive> Ar(IFileManager::Get().CreateFileReader(*InFilePath));
-		if (Ar)
+		FString JsonConfig;
+		if (!FFileHelper::LoadFileToString(JsonConfig, *InFilePath))
+		{
+			return nullptr;
+		}
+
+		const TSharedRef<TJsonReader<>> JsonReader = TJsonReaderFactory<>::Create(JsonConfig);
+
+		TSharedPtr<FJsonObject> JsonObject = MakeShared<FJsonObject>();
+		bool bDeserializeResult = FJsonSerializer::Deserialize(JsonReader, JsonObject);
+		if (!bDeserializeResult || !JsonObject.IsValid())
 		{
-			const TSharedRef<TJsonReader<>> JsonReader = TJsonReaderFactory<>::Create(Ar.Get());
+			// This may have failed because we're reading a file that was saved before encoding was saved explicitly in the file.
+			// So revert to the old read path that uses an explicit FileReader. 
+			if (TUniquePtr<FArchive> Ar = TUniquePtr<FArchive>(IFileManager::Get().CreateFileReader(*InFilePath)))
+			{
+				const TSharedRef<TJsonReader<>> FileReaderJsonReader = TJsonReaderFactory<>::Create(Ar.Get());
+				bDeserializeResult = FJsonSerializer::Deserialize(FileReaderJsonReader, JsonObject) && JsonObject.IsValid();
+			}
+		}
 
-			TSharedPtr<FJsonObject> JsonObject = MakeShared<FJsonObject>();
-			if (FJsonSerializer::Deserialize(JsonReader, JsonObject) && JsonObject.IsValid())
+		if (bDeserializeResult)
+		{
+			int32 Version = 0;
+			if (JsonObject->TryGetNumberField(JsonVersionKey, Version) &&
+				Version <= LiveLinkHubVersion)
 			{
-				int32 Version = 0;
-				if (JsonObject->TryGetNumberField(JsonVersionKey, Version) &&
-					Version <= LiveLinkHubVersion)
-				{
-					return FromJson(JsonObject);
-				}
-
-				UE_LOG(LogLiveLinkHub, Error, TEXT("Could not load config %s because the version field %s was missing or invalid."),
-					*InFilePath, *JsonVersionKey);
+				return FromJson(JsonObject);
 			}
+
+			UE_LOG(LogLiveLinkHub, Error, TEXT("Could not load config %s because the version field %s was missing or invalid."),
+				*InFilePath, *JsonVersionKey);
 		}
 	}
 
