LiveLinkHub - Various bug fixes for Client Filter Presets
- Deleting current Filter preset will now clear it
- Adding filter will no longer create a duplicated LiveLinkHub entry
- Accessing the client filters no longer blocks accessing the hub
- Renamed Disabled AutoConnect to Filters Only

diff --git a/Engine/Plugins/Animation/LiveLinkHub/Source/LiveLinkHub/Private/Clients/LiveLinkHubProvider.cpp b/Engine/Plugins/Animation/LiveLinkHub/Source/LiveLinkHub/Private/Clients/LiveLinkHubProvider.cpp
--- a/Engine/Plugins/Animation/LiveLinkHub/Source/LiveLinkHub/Private/Clients/LiveLinkHubProvider.cpp
+++ b/Engine/Plugins/Animation/LiveLinkHub/Source/LiveLinkHub/Private/Clients/LiveLinkHubProvider.cpp
@@ -1174,12 +1174,27 @@ void FLiveLinkHubProvider::HandleBeaconMessage(const FLiveLinkHubBeaconMessage&
 
 bool FLiveLinkHubProvider::ShouldAutoConnectTo(const FString& RemoteHostname, const FString& IP, const FString& ProjectName) const
 {
+	const ELiveLinkHubAutoConnectMode AutoConnectClients = GetDefault<ULiveLinkHubUserSettings>()->GetAutoConnectMode();
+	if (AutoConnectClients == ELiveLinkHubAutoConnectMode::Disabled)
+	{
+		return false;
+	}
+
+	if (AutoConnectClients == ELiveLinkHubAutoConnectMode::LocalOnly)
+	{
+		// LocalOnly does not apply filters because it's considered an override.
+		return RemoteHostname == GetMachineName();
+	}
+
 	FLiveLinkHubClientFilterPreset Preset = GetMutableDefault<ULiveLinkHubUserSettings>()->GetClientFiltersData_AnyThread();
 
 	bool bIncluded = false;
+	bool bHasAtLeastOneIncludeFilter = false;
 
 	for (const FLiveLinkHubClientTextFilter& Filter : Preset.Filters)
 	{
+		bHasAtLeastOneIncludeFilter |= Filter.Behavior == ELiveLinkHubClientFilterBehavior::Include && (!Filter.Text.IsEmpty() || !Filter.Project.IsEmpty());
+
 		const FString* StringToCompare = Filter.Type == ELiveLinkHubClientFilterType::IP ? &IP : &RemoteHostname;
 		// Allow filtering on project name if IP/Hostname is empty.
 		if (StringToCompare->MatchesWildcard(Filter.Text) || (Filter.Text.IsEmpty() && !Filter.Project.IsEmpty()))
@@ -1202,20 +1217,14 @@ bool FLiveLinkHubProvider::ShouldAutoConnectTo(const FString& RemoteHostname, co
 		}
 	}
 
+
 	if (bIncluded)
 	{
 		return true;
 	}
 
-	const bool bSameHost = RemoteHostname == GetMachineName();
-
-	const ELiveLinkHubAutoConnectMode AutoConnectClients = GetDefault<ULiveLinkHubUserSettings>()->GetAutoConnectMode();
-	if (AutoConnectClients != ELiveLinkHubAutoConnectMode::Disabled)
-	{
-		return AutoConnectClients == ELiveLinkHubAutoConnectMode::All || (AutoConnectClients == ELiveLinkHubAutoConnectMode::LocalOnly && bSameHost);
-	}
-
-	return false;
+	// If there's at least one filter and it wasn't matched, then we discard it.
+	return !bHasAtLeastOneIncludeFilter;
 }
 
 void FLiveLinkHubProvider::OnMessageBusNotification(const FMessageBusNotification& Notification)
@@ -1244,33 +1253,48 @@ void FLiveLinkHubProvider::OnMessageBusNotification(const FMessageBusNotificatio
 
 void FLiveLinkHubProvider::OnFiltersModified()
 {
-	// Remove all clients from session
-	TArray<FLiveLinkHubClientId> ClientIds;
-	{
-		UE::TReadScopeLock Lock(AddressBookLock);
-		ClientIds.Reserve(AddressBook.Num());
-		Algo::Transform(AddressBook, ClientIds, [](const FLiveLinkHubAddressBookEntry& Entry) { return Entry.ClientId; });
-	}
+	TSet<FLiveLinkHubClientId> ClientsToConnect;
+	TSet<FLiveLinkHubClientId> ClientsToDisconnect;
 
-	if (TSharedPtr<ILiveLinkHubSessionManager> Manager = SessionManager.Pin())
 	{
-		for (const FLiveLinkHubClientId& ClientId : ClientIds)
+		UE::TWriteScopeLock Lock(AddressBookLock);
+		for (auto EntryIt = AddressBook.CreateIterator(); EntryIt; ++EntryIt)
 		{
-			Manager->GetCurrentSession()->RemoveClient(ClientId);
+			const bool bShouldConnectResult = ShouldAutoConnectTo(EntryIt->Hostname, EntryIt->IP, EntryIt->ProjectName)
+				&& GetDefault<ULiveLinkHubMessagingSettings>()->CanTransmitTo(FLiveLinkHub::Get()->GetTopologyMode(), EntryIt->TopologyMode);
+
+			if (EntryIt->bConnected)
+			{
+				if (!bShouldConnectResult)
+				{
+					ClientsToDisconnect.Add(EntryIt->ClientId);
+				}
+			}
+			else if (bShouldConnectResult)
+			{
+				ClientsToConnect.Add(EntryIt->ClientId);
+			}
+		}
+
+
+		for (auto EntryIt = AddressBook.CreateIterator(); EntryIt; ++EntryIt)
+		{
+			if (ClientsToConnect.Contains(EntryIt->ClientId))
+			{
+				EntryIt->bSentDiscoveryRequest = true;
+				EntryIt->bDisableAutoconnect = false;
+
+				UE_LOG(LogLiveLinkHub, Verbose, TEXT("Sending Hub Connection request to %s"), *EntryIt->Hostname);
+				SendMessage(MakeLiveLinkHubDiscoveryMessage(), EntryIt->ControlEndpoint, EMessageFlags::Reliable);
+			}
 		}
 	}
 
-	UE::TWriteScopeLock Lock(AddressBookLock);
-	for (auto EntryIt = AddressBook.CreateIterator(); EntryIt; ++EntryIt)
+	if (TSharedPtr<ILiveLinkHubSessionManager> Manager = SessionManager.Pin())
 	{
-		if (!EntryIt->bConnected && ShouldAutoConnectTo(EntryIt->Hostname, EntryIt->IP, EntryIt->ProjectName)
-			&& GetDefault<ULiveLinkHubMessagingSettings>()->CanTransmitTo(FLiveLinkHub::Get()->GetTopologyMode(), EntryIt->TopologyMode))
+		for (const FLiveLinkHubClientId& ClientId : ClientsToDisconnect)
 		{
-			EntryIt->bSentDiscoveryRequest = true;
-			EntryIt->bDisableAutoconnect = false;
-
-			UE_LOG(LogLiveLinkHub, Verbose, TEXT("Sending Hub Connection request to %s"), *EntryIt->Hostname);
-			SendMessage(MakeLiveLinkHubDiscoveryMessage(), EntryIt->ControlEndpoint, EMessageFlags::Reliable);
+			Manager->GetCurrentSession()->RemoveClient(ClientId);
 		}
 	}
 }
diff --git a/Engine/Plugins/Animation/LiveLinkHub/Source/LiveLinkHub/Private/Clients/SLiveLinkHubClientsView.h b/Engine/Plugins/Animation/LiveLinkHub/Source/LiveLinkHub/Private/Clients/SLiveLinkHubClientsView.h
--- a/Engine/Plugins/Animation/LiveLinkHub/Source/LiveLinkHub/Private/Clients/SLiveLinkHubClientsView.h
+++ b/Engine/Plugins/Animation/LiveLinkHub/Source/LiveLinkHub/Private/Clients/SLiveLinkHubClientsView.h
@@ -907,10 +907,8 @@ private:
 				SNew(SLiveLinkHubClientFilters)
 			];
 
-		if (GEditor)
-		{
-			GEditor->EditorAddModalWindow(FiltersMenuWindow.ToSharedRef());
-		}
+
+		FSlateApplication::Get().AddWindow(FiltersMenuWindow.ToSharedRef());
 	}
 
 	/** Saves the current preset to disk. */
diff --git a/Engine/Plugins/Animation/LiveLinkHub/Source/LiveLinkHub/Private/Settings/LiveLinkHubSettings.h b/Engine/Plugins/Animation/LiveLinkHub/Source/LiveLinkHub/Private/Settings/LiveLinkHubSettings.h
--- a/Engine/Plugins/Animation/LiveLinkHub/Source/LiveLinkHub/Private/Settings/LiveLinkHubSettings.h
+++ b/Engine/Plugins/Animation/LiveLinkHub/Source/LiveLinkHub/Private/Settings/LiveLinkHubSettings.h
@@ -352,6 +352,10 @@ public:
 	void DeleteFilterPreset(const FString& PresetName)
 	{
 		FilterPresets.RemoveAll([&PresetName](const FLiveLinkHubClientFilterPreset& InPreset) { return PresetName == InPreset.PresetName; });
+		if (PresetName == CurrentPreset.PresetName)
+		{
+			CurrentPreset = FLiveLinkHubClientFilterPreset{};
+		}
 		SaveConfig();
 	}
 
diff --git a/Engine/Plugins/Animation/LiveLinkHub/Source/LiveLinkHubMessaging/Public/LiveLinkHubMessages.h b/Engine/Plugins/Animation/LiveLinkHub/Source/LiveLinkHubMessaging/Public/LiveLinkHubMessages.h
--- a/Engine/Plugins/Animation/LiveLinkHub/Source/LiveLinkHubMessaging/Public/LiveLinkHubMessages.h
+++ b/Engine/Plugins/Animation/LiveLinkHub/Source/LiveLinkHubMessaging/Public/LiveLinkHubMessages.h
@@ -22,8 +22,8 @@ UENUM()
 enum class ELiveLinkHubAutoConnectMode : uint8
 {
 	Disabled   UMETA(ToolTip = "Don't add any clients automatically."),
-	All 	   UMETA(ToolTip = "Add any client that was found on the network."),
-	LocalOnly  UMETA(ToolTip = "Add any client running on this machine.")
+	All 	   UMETA(ToolTip = "Add any client that was found on the network, unless it doesn't match one of the filters."),
+	LocalOnly  UMETA(ToolTip = "Add any client running on this machine, overriding any current filter.")
 };
 
 /** 
