LiveLinkHub: Fix default selection end time possibly ending early.

Recording length now goes by furthest track timestamp, not frame number.

diff --git a/Engine/Plugins/Animation/LiveLinkHub/Source/LiveLinkHub/Private/Recording/LiveLinkHubPlaybackController.cpp b/Engine/Plugins/Animation/LiveLinkHub/Source/LiveLinkHub/Private/Recording/LiveLinkHubPlaybackController.cpp
--- a/Engine/Plugins/Animation/LiveLinkHub/Source/LiveLinkHub/Private/Recording/LiveLinkHubPlaybackController.cpp
+++ b/Engine/Plugins/Animation/LiveLinkHub/Source/LiveLinkHub/Private/Recording/LiveLinkHubPlaybackController.cpp
@@ -610,12 +610,10 @@ void FLiveLinkHubPlaybackController::SetSelectionEndTime(FQualifiedFrameTime InT
 FQualifiedFrameTime FLiveLinkHubPlaybackController::GetLength() const
 {
 	const FFrameRate FrameRate = GetFrameRate();
+	const double LastTimestamp = RecordingToPlay ? RecordingToPlay->GetLastTimestamp() : 0.f;
 
-	const int32 ExactTotalFrames = RecordingToPlay ? RecordingToPlay->GetMaxFrames() : 1;
-	const int32 LastFrameIdx = ExactTotalFrames - 1;
-	const FFrameNumber LastFrameNumber(LastFrameIdx);
-
-	const FQualifiedFrameTime FrameTime(FFrameTime(LastFrameNumber), FrameRate);
+	const FTimecode LastTimeCode(LastTimestamp, FrameRate, /*bRollover*/ true);
+	const FQualifiedFrameTime FrameTime(LastTimeCode, FrameRate);
 	return FrameTime;
 }
 
