LiveLinkHub - Fix loading source with only Virtual Subject not working

diff --git a/Engine/Plugins/Animation/LiveLink/Source/LiveLink/Private/LiveLinkClient.cpp b/Engine/Plugins/Animation/LiveLink/Source/LiveLink/Private/LiveLinkClient.cpp
--- a/Engine/Plugins/Animation/LiveLink/Source/LiveLink/Private/LiveLinkClient.cpp
+++ b/Engine/Plugins/Animation/LiveLink/Source/LiveLink/Private/LiveLinkClient.cpp
@@ -748,7 +748,11 @@ void FLiveLinkClient::RemoveAllSources()
 	check(Collection);
 	Collection->ForEachSource([](FLiveLinkCollectionSourceItem& SourceItem)
 	{
-		SourceItem.bPendingKill = true;
+		// Don't remove virtual subject sources, or restoring Virtual Subjects will fail because their source doesn't exist anymore.
+		if (!SourceItem.bIsVirtualSource)
+		{
+			SourceItem.bPendingKill = true;
+		}
 	});
 }
 
diff --git a/Engine/Plugins/Animation/LiveLinkHub/Source/LiveLinkHub/Private/Session/LiveLinkHubSessionManager.h b/Engine/Plugins/Animation/LiveLinkHub/Source/LiveLinkHub/Private/Session/LiveLinkHubSessionManager.h
--- a/Engine/Plugins/Animation/LiveLinkHub/Source/LiveLinkHub/Private/Session/LiveLinkHubSessionManager.h
+++ b/Engine/Plugins/Animation/LiveLinkHub/Source/LiveLinkHub/Private/Session/LiveLinkHubSessionManager.h
@@ -561,6 +561,12 @@ private:
 		check(LiveLinkHubClient);
 		
 		const float TimeToWaitForRemoval = bWaitForSourceRemoval ? GetDefault<ULiveLinkHubSettings>()->SourceMaxCleanupTime : 0.f;
+		for (const FLiveLinkSubjectKey& Subject : LiveLinkHubClient->GetSubjects(true, true))
+		{
+			// Make sure we clear all subjects (including virtual subjects).
+			LiveLinkHubClient->RemoveVirtualSubject(Subject);
+		}
+
 		const bool bRemovedAllSources = LiveLinkHubClient->RemoveAllSourcesWithTimeout(TimeToWaitForRemoval);
 		
 		if (!bRemovedAllSources && bWaitForSourceRemoval)
