LiveLinkHub: Fix recordings not always looping.
- Improve frame time calculation and increase precision.
- Buffer range reported should no longer report past max frames.
- Added optional checks on load to validate end frame calculation accuracy.

diff --git a/Engine/Plugins/Animation/LiveLinkHub/Source/LiveLinkHub/Private/Recording/Implementations/LiveLinkUAssetRecording.cpp b/Engine/Plugins/Animation/LiveLinkHub/Source/LiveLinkHub/Private/Recording/Implementations/LiveLinkUAssetRecording.cpp
--- a/Engine/Plugins/Animation/LiveLinkHub/Source/LiveLinkHub/Private/Recording/Implementations/LiveLinkUAssetRecording.cpp
+++ b/Engine/Plugins/Animation/LiveLinkHub/Source/LiveLinkHub/Private/Recording/Implementations/LiveLinkUAssetRecording.cpp
@@ -67,9 +67,17 @@ void ULiveLinkUAssetRecording::PostRename(UObject* OldOuter, const FName OldName
 	EjectAndUnload();
 }
 
+FQualifiedFrameTime ULiveLinkUAssetRecording::GetLastFrameTime() const
+{
+	const FFrameRate Rate = GetGlobalFrameRate();
+	const FFrameTime FrameTime = Rate.AsFrameTime(GetLastTimestamp());
+	const FQualifiedFrameTime QFrameTime(FrameTime, Rate);
+	return QFrameTime;
+}
+
 FFrameRate ULiveLinkUAssetRecording::GetGlobalFrameRate() const
 {
-	return CalculateFrameRate(GetMaxFrames(), GetLastTimestamp());
+	return RecordingLastTimestampFramerate.IsValid() ? RecordingLastTimestampFramerate : FFrameRate(60, 1);
 }
 
 void ULiveLinkUAssetRecording::LoadEntireRecording()
@@ -283,15 +291,14 @@ bool ULiveLinkUAssetRecording::WaitForBufferedFrames(int32 InMinFrame, int32 InM
 			FPlatformProcess::Sleep(0.002);
 		}
 
-		const int32 MaxFrames = GetMaxFrames();
-		if (MaxFrames <= 0)
+		const int32 EndFrame = FMath::Clamp(GetLastFrameTime().Time.RoundToFrame().Value, 0, RecordingMaxFrames - 1);
+		if (EndFrame <= 0)
 		{
 			return true;
 		}
 		
 		// Clamp the frame range to the max possible range. If the selection range extends the actual frame range
 		// then there would be nothing to load.
-		const int32 EndFrame = MaxFrames - 1;
 		InMinFrame = FMath::Clamp(InMinFrame, 0, EndFrame);
 		InMaxFrame = FMath::Clamp(InMaxFrame, 0, EndFrame);
 
@@ -488,9 +495,17 @@ void ULiveLinkUAssetRecording::InitializeNewRecordingData(FLiveLinkUAssetRecordi
 FFrameRate ULiveLinkUAssetRecording::CalculateFrameRate(const int32 InMaxFrames, const double InTime)
 {
 	const double FramesPerSecond = InTime > 0.0 ? InMaxFrames / InTime : 0.0;
-	const int64 FramesPerSecondInt = FMath::RoundToInt(FramesPerSecond);
 
-	return FFrameRate(static_cast<uint32>(FramesPerSecondInt), 1);
+	const double IntegralPart = FMath::FloorToDouble(FramesPerSecond);
+	const double FractionalPart = FramesPerSecond - IntegralPart;
+
+	constexpr uint32 Precision = 100000;
+	const uint32 Divisor = FMath::GreatestCommonDivisor(FMath::RoundToDouble(FractionalPart * Precision), Precision);
+
+	const uint32 Denominator = Precision / Divisor;
+	const uint32 Numerator = IntegralPart * Denominator + FMath::RoundToDouble(FractionalPart * Precision) / Divisor;
+
+	return FFrameRate(Numerator, Denominator);
 }
 
 void ULiveLinkUAssetRecording::SaveFrameData(FArchive* InFileWriter, const FLiveLinkSubjectKey& InSubjectKey, FLiveLinkRecordingBaseDataContainer& InBaseDataContainer)
@@ -880,13 +895,32 @@ bool ULiveLinkUAssetRecording::LoadInitialFrameData(UE::LiveLinkHub::FrameData::
 	}
 
 	OutFrameData.LastTimestamp = OutFrameData.Timestamps.Num () > 0 ? OutFrameData.Timestamps.Last() : 0.0;
+
+	// Calculate frame rate.
+	OutFrameData.LocalFrameRate = CalculateFrameRate(OutFrameData.MaxFrames, OutFrameData.LastTimestamp);
+
 	if (OutFrameData.LastTimestamp > RecordingLastTimestamp)
 	{
 		RecordingLastTimestamp = OutFrameData.LastTimestamp;
-	}
+		RecordingLastTimestampFramerate = OutFrameData.LocalFrameRate;
 
-	// Calculate frame rate.
-	OutFrameData.LocalFrameRate = CalculateFrameRate(OutFrameData.MaxFrames, OutFrameData.LastTimestamp);
+#if LIVELINKHUB_FRAME_DEBUG
+		// Test consistency -- We need to be able to reliably convert to our global rate, which
+		// is determined entirely by the track with the furthest timestamp. This ensures we are always able to
+		// convert to the end frame, and that the end frame equals the original recorded value. Helps to make
+		// sure a recording will always know when its last frame is. This helps both with looping and preventing
+		// freezes waiting for frames to buffer that don't exist.
+		{
+			const FFrameTime FrameTime = RecordingLastTimestampFramerate.AsFrameTime(RecordingLastTimestamp);
+			const FQualifiedFrameTime QFrameTime(FrameTime, RecordingLastTimestampFramerate);
+
+			const int32 CalculatedMaxFrames = QFrameTime.Time.RoundToFrame().Value;
+			const double CalculatedTimestamp = QFrameTime.AsSeconds();
+			ensureAlways(CalculatedMaxFrames == OutFrameData.MaxFrames);
+			ensureAlways(FMath::IsNearlyEqual(CalculatedTimestamp, OutFrameData.LastTimestamp, 0.0001));
+		}
+#endif
+	}
 	
 	return true;
 }
@@ -918,16 +952,18 @@ void ULiveLinkUAssetRecording::LoadFrameData(UE::LiveLinkHub::FrameData::Private
 	// First, localize the frame times to this frame data.
 	if (InFrameData.LocalFrameRate.IsValid())
 	{
-		RequestedStartFrame = StartFrameTime.ConvertTo(InFrameData.LocalFrameRate).FrameNumber.Value;
-		RequestedInitialFrame = StartInitialFrameTime.ConvertTo(InFrameData.LocalFrameRate).FrameNumber.Value;
+		RequestedStartFrame = StartFrameTime.ConvertTo(InFrameData.LocalFrameRate).RoundToFrame().Value;
+		RequestedInitialFrame = StartInitialFrameTime.ConvertTo(InFrameData.LocalFrameRate).RoundToFrame().Value;
 	}
 
 	// Localizing the frame rates may not be enough if the track had inconsistent frames or frame rate. Next, we
 	// need to properly identify the exact local frame to use most closely matching the requested timestamp.
 	// This should produce the same index unless the track had an inconsistent storage and frame rate.
 	// E.g., times 0-1:00 had data, 1:00-1:30 had no data, and 1:30 to 2:00 had data again.
-	int32 StartInitialFrameLocalOffset = 0;
+	int32 InitialFrameLocalOffset = 0;
+	int32 StartFrameLocalOffset = 0;
 	{
+		const int32 GlobalRequestedStartFrame = RequestedStartFrame;
 		const int32 GlobalRequestedInitialFrame = RequestedInitialFrame;
 	
 		RequestedStartFrame = InFrameData.FindFrameFromTimestamp(RequestedStartTimestamp, RequestedStartFrame);
@@ -936,7 +972,8 @@ void ULiveLinkUAssetRecording::LoadFrameData(UE::LiveLinkHub::FrameData::Private
 		// Store any offset from the calculated localized frame to the actual localized frame.
 		// This is needed to report the global buffered frames correctly. As a whole, the recording cares whether the
 		// requested frames should be considered loaded, not the local frames of each track.
-		StartInitialFrameLocalOffset = GlobalRequestedInitialFrame - RequestedInitialFrame;
+		StartFrameLocalOffset = GlobalRequestedStartFrame - RequestedStartFrame;
+		InitialFrameLocalOffset = GlobalRequestedInitialFrame - RequestedInitialFrame;
 	}
 	
 	int32 MaxFrames = InFrameData.MaxFrames;
@@ -944,8 +981,10 @@ void ULiveLinkUAssetRecording::LoadFrameData(UE::LiveLinkHub::FrameData::Private
 	{
 		if (RequestedFramesToLoad > 0)
 		{
-			// Don't go past requested frames or max frames.
-			MaxFrames = FMath::Min(MaxFrames, RequestedStartFrame + RequestedFramesToLoad + 1);
+			// Don't go past requested frames or max frames. We add back in StartFrameLocalOffset because the start and initial
+			// frames have been localized to the track, but the total frames we need may go to the end of the track, and we need to honor
+			// the original request.
+			MaxFrames = FMath::Min(MaxFrames, RequestedStartFrame + RequestedFramesToLoad + StartFrameLocalOffset + 1);
 			// Ensure we're always at least streaming in the initial frame. This is usually the case, but depending on the track's
 			// frames and how the RequestedFramesToLoad was divided, it may cut off the initial frame if it's the last frame of the track.
 			MaxFrames = FMath::Max(MaxFrames, RequestedInitialFrame + 1);
@@ -1304,7 +1343,7 @@ void ULiveLinkUAssetRecording::LoadFrameData(UE::LiveLinkHub::FrameData::Private
 		IterationData.LastLoadedLeftFrame = LastLoadedLeftFrame;
 		IterationData.LastLoadedRightFrame = LastLoadedRightFrame;
 
-		InDataContainer.LocalFrameOffset = StartInitialFrameLocalOffset;
+		InDataContainer.LocalFrameOffset = InitialFrameLocalOffset;
 		
 		// Determine the iteration status, which will impact the overall async load loop.
 		if (RightFrameIdx >= MaxFrames && LeftFrameIdx <= RequestedStartFrame)
@@ -1534,6 +1573,10 @@ void ULiveLinkUAssetRecording::UpdateBufferedFrames()
 	// the ranges as viewed by the recording/scrubber.
 	BufferedFrameRanges.Reset();
 	InactiveBufferedFrameRanges.Reset();
+
+	const FQualifiedFrameTime LastFrameTime = GetLastFrameTime();
+	const int32 LastFrameValue = LastFrameTime.Time.RoundToFrame().Value;
+	
 	for (TTuple<FLiveLinkSubjectKey, FrameData::Private::FFrameMetaData>& FrameDataKeyVal : FrameFileData)
 	{
 		FrameData::Private::FFrameMetaData& FrameData = FrameDataKeyVal.Value;
@@ -1541,10 +1584,19 @@ void ULiveLinkUAssetRecording::UpdateBufferedFrames()
 		for (TRange<int32>& Range : CacheBufferRanges)
 		{
 			Range = RangeHelpers::Private::ConvertRangeFrameRate(Range,FrameData.LocalFrameRate, GlobalFrameRate);
+			if (Range.GetUpperBoundValue() > LastFrameValue)
+			{
+				Range.SetUpperBoundValue(LastFrameValue);
+			}
 		}
 		
 		TRange<int32> ActiveBufferFrames = RecordingData.FrameData.FindChecked(*FrameData.FrameDataSubjectKey.Get()).GetBufferedFrames(/*bIncludeOffset*/ true);
 		ActiveBufferFrames = RangeHelpers::Private::ConvertRangeFrameRate(ActiveBufferFrames, FrameData.LocalFrameRate, GlobalFrameRate);
+
+		if (ActiveBufferFrames.GetUpperBoundValue() > LastFrameValue)
+		{
+			ActiveBufferFrames.SetUpperBoundValue(LastFrameValue);
+		}
 		
 		InactiveBufferedFrameRanges.Append(MoveTemp(CacheBufferRanges));
 		BufferedFrameRanges.Add(MoveTemp(ActiveBufferFrames));
diff --git a/Engine/Plugins/Animation/LiveLinkHub/Source/LiveLinkHub/Private/Recording/Implementations/LiveLinkUAssetRecording.h b/Engine/Plugins/Animation/LiveLinkHub/Source/LiveLinkHub/Private/Recording/Implementations/LiveLinkUAssetRecording.h
--- a/Engine/Plugins/Animation/LiveLinkHub/Source/LiveLinkHub/Private/Recording/Implementations/LiveLinkUAssetRecording.h
+++ b/Engine/Plugins/Animation/LiveLinkHub/Source/LiveLinkHub/Private/Recording/Implementations/LiveLinkUAssetRecording.h
@@ -38,6 +38,7 @@ public:
 	virtual bool IsSavingRecordingData() const override { return bIsSavingRecordingData; }
 	virtual int32 GetMaxFrames() const override { return RecordingMaxFrames; }
 	virtual double GetLastTimestamp() const override { return RecordingLastTimestamp; }
+	virtual FQualifiedFrameTime GetLastFrameTime() const override;
 	virtual FFrameRate GetGlobalFrameRate() const override;
 	virtual void LoadEntireRecording() override;
 	virtual bool IsLoadingEntireRecording() const override { return bIsPerformingFullLoad;}
@@ -213,6 +214,13 @@ private:
 	/** The last time stamp of the recording.  */
 	std::atomic<double> RecordingLastTimestamp = 0.0;
 
+	/**
+	 * The framerate of the track the last timestamp belongs to.
+	 * We store this because we need to calculate all frames conversion between tracks to this
+	 * "global" rate.
+	 */
+	FFrameRate RecordingLastTimestampFramerate;
+
 	/** Frames buffered, divided into ranges. */
 	UE::LiveLinkHub::RangeHelpers::Private::TRangeArray<int32> BufferedFrameRanges;
 	/** Ranges that are buffered but considered inactive. */
diff --git a/Engine/Plugins/Animation/LiveLinkHub/Source/LiveLinkHub/Private/Recording/Implementations/LiveLinkUAssetRecordingPlayer.cpp b/Engine/Plugins/Animation/LiveLinkHub/Source/LiveLinkHub/Private/Recording/Implementations/LiveLinkUAssetRecordingPlayer.cpp
--- a/Engine/Plugins/Animation/LiveLinkHub/Source/LiveLinkHub/Private/Recording/Implementations/LiveLinkUAssetRecordingPlayer.cpp
+++ b/Engine/Plugins/Animation/LiveLinkHub/Source/LiveLinkHub/Private/Recording/Implementations/LiveLinkUAssetRecordingPlayer.cpp
@@ -216,11 +216,11 @@ int32 FLiveLinkPlaybackTrack::ConvertFrameTimeToFrameIndex(const FQualifiedFrame
 	{
 		// Localize the frame time relative to this track's framerate.
 		const FFrameTime LocalizedFrameTime = InFrameTime.ConvertTo(LocalFrameRate);
-		LocalizedIndex = LocalizedFrameTime.GetFrame().Value;
+		LocalizedIndex = LocalizedFrameTime.RoundToFrame().Value;
 	}
 	else
 	{
-		LocalizedIndex = InFrameTime.Time.GetFrame().Value;
+		LocalizedIndex = InFrameTime.Time.RoundToFrame().Value;
 	}
 
 	int32 CurrentIndex = GetRelativeIndex(LocalizedIndex);
diff --git a/Engine/Plugins/Animation/LiveLinkHub/Source/LiveLinkHub/Private/Recording/Implementations/LiveLinkUAssetRecordingPlayer.h b/Engine/Plugins/Animation/LiveLinkHub/Source/LiveLinkHub/Private/Recording/Implementations/LiveLinkUAssetRecordingPlayer.h
--- a/Engine/Plugins/Animation/LiveLinkHub/Source/LiveLinkHub/Private/Recording/Implementations/LiveLinkUAssetRecordingPlayer.h
+++ b/Engine/Plugins/Animation/LiveLinkHub/Source/LiveLinkHub/Private/Recording/Implementations/LiveLinkUAssetRecordingPlayer.h
@@ -96,7 +96,7 @@ public:
 	
 	virtual TArray<FLiveLinkRecordedFrame> FetchNextFramesAtTimestamp(const FQualifiedFrameTime& InFrameTime) override
 	{
-		if (StreamPlayback(InFrameTime.Time.GetFrame().Value))
+		if (StreamPlayback(InFrameTime.Time.RoundToFrame().Value))
 		{
 			return CurrentRecordingPlayback.FetchNextFrames(InFrameTime);
 		}
@@ -105,7 +105,7 @@ public:
 
 	virtual TArray<FLiveLinkRecordedFrame> FetchPreviousFramesAtTimestamp(const FQualifiedFrameTime& InFrameTime) override
 	{
-		if (StreamPlayback(InFrameTime.Time.GetFrame().Value))
+		if (StreamPlayback(InFrameTime.Time.RoundToFrame().Value))
 		{
 			return CurrentRecordingPlayback.FetchPreviousFrames(InFrameTime);
 		}
@@ -114,7 +114,7 @@ public:
 
 	virtual TArray<FLiveLinkRecordedFrame> FetchNextFramesAtIndex(const FQualifiedFrameTime& InFrameTime) override
 	{
-		if (StreamPlayback(InFrameTime.Time.GetFrame().Value))
+		if (StreamPlayback(InFrameTime.Time.RoundToFrame().Value))
 		{
 			return CurrentRecordingPlayback.FetchNextFramesAtIndex(InFrameTime);
 		}
diff --git a/Engine/Plugins/Animation/LiveLinkHub/Source/LiveLinkHub/Private/Recording/LiveLinkHubPlaybackController.cpp b/Engine/Plugins/Animation/LiveLinkHub/Source/LiveLinkHub/Private/Recording/LiveLinkHubPlaybackController.cpp
--- a/Engine/Plugins/Animation/LiveLinkHub/Source/LiveLinkHub/Private/Recording/LiveLinkHubPlaybackController.cpp
+++ b/Engine/Plugins/Animation/LiveLinkHub/Source/LiveLinkHub/Private/Recording/LiveLinkHubPlaybackController.cpp
@@ -609,12 +609,12 @@ void FLiveLinkHubPlaybackController::SetSelectionEndTime(FQualifiedFrameTime InT
 
 FQualifiedFrameTime FLiveLinkHubPlaybackController::GetLength() const
 {
-	const FFrameRate FrameRate = GetFrameRate();
-	const double LastTimestamp = RecordingToPlay ? RecordingToPlay->GetLastTimestamp() : 0.f;
+	if (RecordingToPlay)
+	{
+		return RecordingToPlay->GetLastFrameTime();
+	}
 
-	const FTimecode LastTimeCode(LastTimestamp, FrameRate, /*bRollover*/ true);
-	const FQualifiedFrameTime FrameTime(LastTimeCode, FrameRate);
-	return FrameTime;
+	return FQualifiedFrameTime();
 }
 
 FQualifiedFrameTime FLiveLinkHubPlaybackController::GetCurrentTime() const
@@ -691,9 +691,9 @@ uint32 FLiveLinkHubPlaybackController::Run()
 					const double TimeDilatedDelta = (FPlatformTime::Seconds() - PlaybackStartTime) * PlayRate;
 					double Position = bIsReverse ? StartTimestamp - TimeDilatedDelta : StartTimestamp + TimeDilatedDelta;
 					Position = FMath::Clamp(Position, GetSelectionStartTime().AsSeconds(), GetSelectionEndTime().AsSeconds());
-					Playhead->SetValue(FQualifiedFrameTime(FFrameTime::FromDecimal(Position * GetFrameRate().Numerator), GetFrameRate()));
+					const FFrameRate FrameRate = GetFrameRate();
+					Playhead->SetValue(FQualifiedFrameTime(FrameRate.AsFrameTime(Position), FrameRate));
 				};
-
 				SetPlayhead();
 			
 				// Don't sleep if we pushed frames since that can take a small amount of time.
diff --git a/Engine/Plugins/Animation/LiveLinkHub/Source/LiveLinkHub/Private/Recording/LiveLinkRecording.h b/Engine/Plugins/Animation/LiveLinkHub/Source/LiveLinkHub/Private/Recording/LiveLinkRecording.h
--- a/Engine/Plugins/Animation/LiveLinkHub/Source/LiveLinkHub/Private/Recording/LiveLinkRecording.h
+++ b/Engine/Plugins/Animation/LiveLinkHub/Source/LiveLinkHub/Private/Recording/LiveLinkRecording.h
@@ -4,6 +4,7 @@
 
 #include "LiveLinkPreset.h"
 #include "Misc/FrameRate.h"
+#include "Misc/QualifiedFrameTime.h"
 #include "UObject/Object.h"
 
 #include "LiveLinkRecording.generated.h"
@@ -31,6 +32,9 @@ public:
 	/** Return the last timestamp of the recording. */
 	virtual double GetLastTimestamp() const { return 0.0; }
 
+	/** Return the last frame time. */
+	virtual FQualifiedFrameTime GetLastFrameTime() const { return FQualifiedFrameTime(); }
+
 	/** Retrieve the generalized framerate of this recording. */
 	virtual FFrameRate GetGlobalFrameRate() const { return FFrameRate(); }
 
diff --git a/Engine/Plugins/Animation/LiveLinkHub/Source/LiveLinkHub/Private/UI/Widgets/SLiveLinkHubPlaybackWidget.cpp b/Engine/Plugins/Animation/LiveLinkHub/Source/LiveLinkHub/Private/UI/Widgets/SLiveLinkHubPlaybackWidget.cpp
--- a/Engine/Plugins/Animation/LiveLinkHub/Source/LiveLinkHub/Private/UI/Widgets/SLiveLinkHubPlaybackWidget.cpp
+++ b/Engine/Plugins/Animation/LiveLinkHub/Source/LiveLinkHub/Private/UI/Widgets/SLiveLinkHubPlaybackWidget.cpp
@@ -97,8 +97,8 @@ void SLiveLinkHubPlaybackWidget::Construct(const FArguments& InArgs)
 					.ScrubPosition_Lambda([this]()
 					{
 						// TimeSlider needs actual time in seconds, not the frame time.
-						const double Seconds = GetCurrentTime() / GetFrameRate().Numerator;
-						return Seconds;
+						const FQualifiedFrameTime Now = GetCurrentTime();
+						return Now.AsSeconds();
 					})
 					.ViewRange(this, &SLiveLinkHubPlaybackWidget::GetViewRange)
 					.OnViewRangeChanged(this, &SLiveLinkHubPlaybackWidget::SetViewRange)
@@ -111,7 +111,7 @@ void SLiveLinkHubPlaybackWidget::Construct(const FArguments& InArgs)
 									if (bIsScrubbing)
 									{
 										//  Convert time in seconds to frame time as double.
-										SetCurrentTime(NewScrubTime * GetFrameRate().Numerator);
+										SetCurrentTime(NewScrubTime);
 									}
 								})
 					)
@@ -307,7 +307,7 @@ void SLiveLinkHubPlaybackWidget::Construct(const FArguments& InArgs)
 				[
 					SNew(SSpinBox<double>)
 					.IsEnabled(IsPlaybackEnabled)
-					.Value(this, &SLiveLinkHubPlaybackWidget::GetCurrentTime)
+					.Value(this, &SLiveLinkHubPlaybackWidget::GetCurrentTimeDouble)
 					.ToolTipText(LOCTEXT("CurrentTime", "Current playback time"))
 					.OnValueChanged(this, &SLiveLinkHubPlaybackWidget::SetCurrentTime)
 					.MinValue(TOptional<double>())
@@ -530,26 +530,34 @@ void SLiveLinkHubPlaybackWidget::SetCurrentTime(double InTime)
 	OnSetCurrentTimeDelegate.Execute(FrameTime);
 }
 
-double SLiveLinkHubPlaybackWidget::GetCurrentTime() const
+double SLiveLinkHubPlaybackWidget::GetCurrentTimeDouble() const
 {
 	check(OnGetCurrentTimeDelegate.IsBound());
 	const FQualifiedFrameTime FrameTime = OnGetCurrentTimeDelegate.Execute();
 		
-	return FrameTime.Time.AsDecimal();
+	return FrameTime.AsSeconds();
+}
+
+FQualifiedFrameTime SLiveLinkHubPlaybackWidget::GetCurrentTime() const
+{
+	check(OnGetCurrentTimeDelegate.IsBound());
+	const FQualifiedFrameTime FrameTime = OnGetCurrentTimeDelegate.Execute();
+		
+	return FrameTime;
 }
 
 double SLiveLinkHubPlaybackWidget::GetTotalLength() const
 {
 	check(OnGetTotalLengthDelegate.IsBound());
 	const FQualifiedFrameTime FrameTime = OnGetTotalLengthDelegate.Execute();
-	return FrameTime.Time.AsDecimal();
+	return FrameTime.AsSeconds();
 }
 
 double SLiveLinkHubPlaybackWidget::GetSelectionStartTime() const
 {
 	check(OnGetSelectionStartTimeDelegate.IsBound());
 	const FQualifiedFrameTime FrameTime = OnGetSelectionStartTimeDelegate.Execute();
-	return FrameTime.Time.AsDecimal();
+	return FrameTime.AsSeconds();
 }
 
 void SLiveLinkHubPlaybackWidget::SetSelectionStartTime(double InTime)
@@ -564,7 +572,7 @@ double SLiveLinkHubPlaybackWidget::GetSelectionEndTime() const
 	check(OnGetSelectionEndTimeDelegate.IsBound());
 	const FQualifiedFrameTime FrameTime = OnGetSelectionEndTimeDelegate.Execute();
 
-	return FrameTime.Time.AsDecimal();
+	return FrameTime.AsSeconds();
 }
 
 void SLiveLinkHubPlaybackWidget::SetSelectionEndTime(double InTime)
@@ -588,9 +596,8 @@ void SLiveLinkHubPlaybackWidget::SetViewRange(TRange<double> InRange)
 
 TRange<double> SLiveLinkHubPlaybackWidget::GetClampRange() const
 {
-	const FFrameRate FrameRate = GetFrameRate();
-	const double Start = GetSelectionStartTime() / FrameRate.Numerator;
-	const double End = GetSelectionEndTime() / FrameRate.Numerator;
+	const double Start = GetSelectionStartTime();
+	const double End = GetSelectionEndTime();
 	return UE::LiveLinkHub::RangeHelpers::Private::MakeInclusiveRange(Start, End);
 }
 
@@ -664,7 +671,8 @@ FFrameRate SLiveLinkHubPlaybackWidget::GetFrameRate() const
 
 FQualifiedFrameTime SLiveLinkHubPlaybackWidget::SecondsToFrameTime(double InTime) const
 {
-	return FQualifiedFrameTime(FFrameTime::FromDecimal(InTime), GetFrameRate());
+	const FFrameRate FrameRate = GetFrameRate();
+	return FQualifiedFrameTime(FrameRate.AsFrameTime(InTime), FrameRate);
 }
 
 const FSlateBrush* SLiveLinkHubPlaybackWidget::GetPlayForwardIcon() const
diff --git a/Engine/Plugins/Animation/LiveLinkHub/Source/LiveLinkHub/Private/UI/Widgets/SLiveLinkHubPlaybackWidget.h b/Engine/Plugins/Animation/LiveLinkHub/Source/LiveLinkHub/Private/UI/Widgets/SLiveLinkHubPlaybackWidget.h
--- a/Engine/Plugins/Animation/LiveLinkHub/Source/LiveLinkHub/Private/UI/Widgets/SLiveLinkHubPlaybackWidget.h
+++ b/Engine/Plugins/Animation/LiveLinkHub/Source/LiveLinkHub/Private/UI/Widgets/SLiveLinkHubPlaybackWidget.h
@@ -134,7 +134,8 @@ private:
 	void SetCurrentTime(double InTime);
 
 	/** The current playhead time. */
-	double GetCurrentTime() const;
+	double GetCurrentTimeDouble() const;
+	FQualifiedFrameTime GetCurrentTime() const;
 	/** The total length of the recording. */
 	double GetTotalLength() const;
 
