LiveLinkHub - Add a setting to disable crash recovery

diff --git a/Engine/Plugins/Animation/LiveLinkHub/Source/LiveLinkHub/Private/LiveLinkHub.cpp b/Engine/Plugins/Animation/LiveLinkHub/Source/LiveLinkHub/Private/LiveLinkHub.cpp
--- a/Engine/Plugins/Animation/LiveLinkHub/Source/LiveLinkHub/Private/LiveLinkHub.cpp
+++ b/Engine/Plugins/Animation/LiveLinkHub/Source/LiveLinkHub/Private/LiveLinkHub.cpp
@@ -217,7 +217,10 @@ void FLiveLinkHub::InitializePostEngineInit()
 
 	GIsRunning = true;
 
-	CrashRecoveryHandler = MakePimpl<FLiveLinkHubCrashRecovery>();
+	if (GetDefault<ULiveLinkHubSettings>()->bEnableCrashRecovery)
+	{
+		CrashRecoveryHandler = MakePimpl<FLiveLinkHubCrashRecovery>();
+	}
 
 	FString SessionPath;
 	if (FParse::Value(FCommandLine::Get(), TEXT("SessionPath="), SessionPath) && SessionPath.Len())
diff --git a/Engine/Plugins/Animation/LiveLinkHub/Source/LiveLinkHub/Private/Settings/LiveLinkHubSettings.h b/Engine/Plugins/Animation/LiveLinkHub/Source/LiveLinkHub/Private/Settings/LiveLinkHubSettings.h
--- a/Engine/Plugins/Animation/LiveLinkHub/Source/LiveLinkHub/Private/Settings/LiveLinkHubSettings.h
+++ b/Engine/Plugins/Animation/LiveLinkHub/Source/LiveLinkHub/Private/Settings/LiveLinkHubSettings.h
@@ -70,6 +70,10 @@ public:
 	UPROPERTY(config, EditAnywhere, Category = "LiveLinkHub")
 	bool bConfirmClose = true;
 
+	/** Whether to enable Live Link Hub's crash recovery system. */
+	UPROPERTY(config, EditAnywhere, Category = "LiveLinkHub")
+	bool bEnableCrashRecovery = true;
+
 	/** Whether to show the app's frame rate in the top right corner. */
 	UPROPERTY(config, EditAnywhere, Category = "LiveLinkHub")
 	bool bShowFrameRate = false;
