LiveLinkHub: Add delay when scanning asset registry after a directory change detection to help cases where the file system may still be locking the file after a manual move.
- Add recording context option to rescan the recording list manually.
- Prevent loading or creating a config if playback session is active.

diff --git a/Engine/Plugins/Animation/LiveLinkHub/Source/LiveLinkHub/Private/LiveLinkHub.cpp b/Engine/Plugins/Animation/LiveLinkHub/Source/LiveLinkHub/Private/LiveLinkHub.cpp
--- a/Engine/Plugins/Animation/LiveLinkHub/Source/LiveLinkHub/Private/LiveLinkHub.cpp
+++ b/Engine/Plugins/Animation/LiveLinkHub/Source/LiveLinkHub/Private/LiveLinkHub.cpp
@@ -406,8 +406,10 @@ void FLiveLinkHub::MapToolkitCommands()
 	const FLiveLinkHubCommands& Commands = FLiveLinkHubCommands::Get();
 	TSharedPtr<FUICommandList> CommandList = GetToolkitCommands();
 
-	CommandList->MapAction(Commands.NewConfig, FExecuteAction::CreateSP(this, &FLiveLinkHub::NewConfig));
-	CommandList->MapAction(Commands.OpenConfig, FExecuteAction::CreateSP(this, &FLiveLinkHub::OpenConfig));
+	CommandList->MapAction(Commands.NewConfig, FExecuteAction::CreateSP(this, &FLiveLinkHub::NewConfig),
+		FCanExecuteAction::CreateSP(this, &FLiveLinkHub::CanCreateOrOpenConfig));
+	CommandList->MapAction(Commands.OpenConfig, FExecuteAction::CreateSP(this, &FLiveLinkHub::OpenConfig),
+		FCanExecuteAction::CreateSP(this, &FLiveLinkHub::CanCreateOrOpenConfig));
 	CommandList->MapAction(Commands.SaveConfigAs, FExecuteAction::CreateSP(this, &FLiveLinkHub::SaveConfigAs));
 	CommandList->MapAction(Commands.SaveConfig, FExecuteAction::CreateSP(this, &FLiveLinkHub::SaveConfig),
 		FCanExecuteAction::CreateSP(this, &FLiveLinkHub::CanSaveConfig));
@@ -491,6 +493,11 @@ void FLiveLinkHub::OpenConfig()
 	SessionManager->RestoreSession();
 }
 
+bool FLiveLinkHub::CanCreateOrOpenConfig() const
+{
+	return !PlaybackController->IsInPlayback();
+}
+
 void FLiveLinkHub::OpenLogsFolder()
 {
 	const FString AbsoluteFilePath = FPaths::ConvertRelativePathToFull(FGenericPlatformOutputDevices::GetAbsoluteLogFilename());
diff --git a/Engine/Plugins/Animation/LiveLinkHub/Source/LiveLinkHub/Private/LiveLinkHub.h b/Engine/Plugins/Animation/LiveLinkHub/Source/LiveLinkHub/Private/LiveLinkHub.h
--- a/Engine/Plugins/Animation/LiveLinkHub/Source/LiveLinkHub/Private/LiveLinkHub.h
+++ b/Engine/Plugins/Animation/LiveLinkHub/Source/LiveLinkHub/Private/LiveLinkHub.h
@@ -145,6 +145,9 @@ private:
 	/** Open an existing config. */
 	void OpenConfig();
 
+	/** If a config can be created or loaded. */
+	bool CanCreateOrOpenConfig() const;
+
 	/** Open the folder containing logs for the app. */
 	void OpenLogsFolder();
 
diff --git a/Engine/Plugins/Animation/LiveLinkHub/Source/LiveLinkHub/Private/LiveLinkHubApplication.cpp b/Engine/Plugins/Animation/LiveLinkHub/Source/LiveLinkHub/Private/LiveLinkHubApplication.cpp
--- a/Engine/Plugins/Animation/LiveLinkHub/Source/LiveLinkHub/Private/LiveLinkHubApplication.cpp
+++ b/Engine/Plugins/Animation/LiveLinkHub/Source/LiveLinkHub/Private/LiveLinkHubApplication.cpp
@@ -442,13 +442,24 @@ void FLiveLinkHubApplication::RegisterFileMenu(UToolMenu* Menu)
 					auto CreateMenuEntry = [](const FLiveLinkHubSessionFile& SessionFile)
 					{
 						TSharedPtr<ILiveLinkHubSessionManager> SessionManager = FLiveLinkHub::Get()->GetSessionManager();
-
+							
+						FToolUIAction ToolUIAction(
+						FToolMenuExecuteAction::CreateLambda([SessionFile](const FToolMenuContext&)
+						{
+							FLiveLinkHub::Get()->GetSessionManager()->RestoreSession(SessionFile.FilePath);
+						}));
+						
+						ToolUIAction.CanExecuteAction = FToolMenuCanExecuteAction::CreateLambda([](const FToolMenuContext&)
+						{
+							return !FLiveLinkHub::Get()->GetPlaybackController()->IsInPlayback();
+						});
+							
 						FToolMenuEntry Entry = FToolMenuEntry::InitMenuEntry(
 							*SessionFile.FileName,
 							FText::FromString(SessionFile.FileName),
 							FText::FromString(SessionFile.LastModificationDate.ToString()),
 							FSlateIcon(),
-							FExecuteAction::CreateLambda([SessionFile]() { FLiveLinkHub::Get()->GetSessionManager()->RestoreSession(SessionFile.FilePath); }),
+							ToolUIAction,
 							EUserInterfaceActionType::None
 						);
 
diff --git a/Engine/Plugins/Animation/LiveLinkHub/Source/LiveLinkHub/Private/Recording/LiveLinkHubRecordingListController.h b/Engine/Plugins/Animation/LiveLinkHub/Source/LiveLinkHub/Private/Recording/LiveLinkHubRecordingListController.h
--- a/Engine/Plugins/Animation/LiveLinkHub/Source/LiveLinkHub/Private/Recording/LiveLinkHubRecordingListController.h
+++ b/Engine/Plugins/Animation/LiveLinkHub/Source/LiveLinkHub/Private/Recording/LiveLinkHubRecordingListController.h
@@ -53,6 +53,11 @@ public:
 		{
 			DirectoryWatcher->Get()->UnregisterDirectoryChangedCallback_Handle(ContentDirectory, DirectoryChangedHandle);
 		}
+
+		if (GEditor && GEditor->IsTimerManagerValid())
+		{
+			GEditor->GetTimerManager()->ClearTimer(DirectoryScanTimerHandle);
+		}
 	}
 
 	/** Create the list's widget. */
@@ -96,12 +101,30 @@ private:
 	/** Reacts to changes in the content directory to rescan it to update the recording list. */
 	void OnContentDirectoryChanged(const TArray<struct FFileChangeData>& FileChangeData)
 	{
-		TArray<FString> FilePaths;
-		FilePaths.Reserve(FileChangeData.Num());
-
-		Algo::Transform(FileChangeData, FilePaths, [](const FFileChangeData& Data) { return Data.Filename; });
-		IAssetRegistry& AssetRegistry = FModuleManager::LoadModuleChecked<FAssetRegistryModule>(AssetRegistryConstants::ModuleName).Get();
-		AssetRegistry.ScanPathsSynchronous({ FilePaths }, /*bForceRescan=*/ true);
+		auto ScanChanges = [this, FileChangeData]()
+		{
+			if (IsEngineExitRequested() || !FLiveLinkHub::Get().IsValid())
+			{
+				return;
+			}
+				
+			TArray<FString> FilePaths;
+			FilePaths.Reserve(FileChangeData.Num());
+
+			Algo::Transform(FileChangeData, FilePaths, [](const FFileChangeData& Data) { return Data.Filename; });
+			IAssetRegistry& AssetRegistry = FModuleManager::LoadModuleChecked<FAssetRegistryModule>(AssetRegistryConstants::ModuleName).Get();
+			AssetRegistry.ScanPathsSynchronous({ FilePaths }, /*bForceRescan=*/ true);
+		};
+		
+		if (GEditor && GEditor->IsTimerManagerValid())
+		{
+			GEditor->GetTimerManager()->ClearTimer(DirectoryScanTimerHandle);
+			GEditor->GetTimerManager()->SetTimer(DirectoryScanTimerHandle, ScanChanges, 1.f, false);
+		}
+		else
+		{
+			ScanChanges();
+		}
 	}
 
 private:
@@ -111,6 +134,9 @@ private:
 	FDelegateHandle DirectoryChangedHandle;
 	/** Path to the content directory. Used to react to changes to that directory, such as new files being added. */
 	FString ContentDirectory;
+
+	/** Timer handle used for scanning directory changes. */
+	FTimerHandle DirectoryScanTimerHandle;
 };
 
 #undef LOCTEXT_NAMESPACE /* LiveLinkHub.RecordingList */
diff --git a/Engine/Plugins/Animation/LiveLinkHub/Source/LiveLinkHub/Private/Recording/SLiveLinkHubRecordingListView.h b/Engine/Plugins/Animation/LiveLinkHub/Source/LiveLinkHub/Private/Recording/SLiveLinkHubRecordingListView.h
--- a/Engine/Plugins/Animation/LiveLinkHub/Source/LiveLinkHub/Private/Recording/SLiveLinkHubRecordingListView.h
+++ b/Engine/Plugins/Animation/LiveLinkHub/Source/LiveLinkHub/Private/Recording/SLiveLinkHubRecordingListView.h
@@ -369,6 +369,22 @@ private:
 		}
 		MenuBuilder.EndSection();
 
+		MenuBuilder.BeginSection(TEXT("RecordingList"), LOCTEXT("RecordingListSectionLabel", "Recording List"));
+		{
+			MenuBuilder.AddMenuEntry(
+			LOCTEXT("RefreshRecordings", "Refresh Recordings"),
+			LOCTEXT("RefreshRecordingsTooltip", "Rescan the directory list."),
+			FSlateIcon(FAppStyle::GetAppStyleSetName(), "Icons.Refresh"),
+			FUIAction(
+				FExecuteAction::CreateLambda([] ()
+				{
+					const FAssetRegistryModule& AssetRegistryModule = FModuleManager::LoadModuleChecked<FAssetRegistryModule>("AssetRegistry");
+					AssetRegistryModule.Get().ScanPathsSynchronous({ TEXT("/Game") }, true);
+				}))
+			);
+		}
+		MenuBuilder.EndSection();
+
 		return MenuBuilder.MakeWidget();
 	}
 	
