LiveLinkHub - Fix crash when selecting subject with Subject Details closed

diff --git a/Engine/Plugins/Animation/LiveLinkHub/Source/LiveLinkHub/Private/Subjects/LiveLinkHubSubjectController.h b/Engine/Plugins/Animation/LiveLinkHub/Source/LiveLinkHub/Private/Subjects/LiveLinkHubSubjectController.h
--- a/Engine/Plugins/Animation/LiveLinkHub/Source/LiveLinkHub/Private/Subjects/LiveLinkHubSubjectController.h
+++ b/Engine/Plugins/Animation/LiveLinkHub/Source/LiveLinkHub/Private/Subjects/LiveLinkHubSubjectController.h
@@ -40,7 +40,10 @@ public:
 	/** Set the displayed subject in the subject view. */
 	void SetSubject(const FLiveLinkSubjectKey& Subject) const
 	{
-		SubjectsView->SetSubject(Subject);
+		if (SubjectsView)
+		{
+			SubjectsView->SetSubject(Subject);
+		}
 	}
 
 	/** Handle updating the subject details when the session has been swapped out for a different one. */
