LiveLinkHub - Fix Save Current Preset overriding current filter

diff --git a/Engine/Plugins/Animation/LiveLinkHub/Source/LiveLinkHub/Private/Clients/SLiveLinkHubClientsView.h b/Engine/Plugins/Animation/LiveLinkHub/Source/LiveLinkHub/Private/Clients/SLiveLinkHubClientsView.h
--- a/Engine/Plugins/Animation/LiveLinkHub/Source/LiveLinkHub/Private/Clients/SLiveLinkHubClientsView.h
+++ b/Engine/Plugins/Animation/LiveLinkHub/Source/LiveLinkHub/Private/Clients/SLiveLinkHubClientsView.h
@@ -1032,7 +1032,7 @@ private:
 	void HandleSaveCurrentPreset()
 	{
 		ULiveLinkHubUserSettings* Settings = GetMutableDefault<ULiveLinkHubUserSettings>();
-		Settings->SaveClientsToPreset(Settings->CurrentPreset.PresetName);
+		Settings->SaveCurrentPreset();
 	}
 
 	/** Open a window that hosts the filters menu. */
@@ -1106,7 +1106,7 @@ private:
 		
 		if ((NamePickerResult == 0 || bConfirmedPresetName) && !PresetName.IsEmpty())
 		{
-			GetMutableDefault<ULiveLinkHubUserSettings>()->SaveClientsToPreset(PresetName);
+			GetMutableDefault<ULiveLinkHubUserSettings>()->SavePresetAs(PresetName);
 		}
 	}
 
diff --git a/Engine/Plugins/Animation/LiveLinkHub/Source/LiveLinkHub/Private/Settings/LiveLinkHubSettings.cpp b/Engine/Plugins/Animation/LiveLinkHub/Source/LiveLinkHub/Private/Settings/LiveLinkHubSettings.cpp
--- a/Engine/Plugins/Animation/LiveLinkHub/Source/LiveLinkHub/Private/Settings/LiveLinkHubSettings.cpp
+++ b/Engine/Plugins/Animation/LiveLinkHub/Source/LiveLinkHub/Private/Settings/LiveLinkHubSettings.cpp
@@ -60,6 +60,15 @@ TObjectPtr<ULiveLinkHubNamingTokens> ULiveLinkHubSettings::GetNamingTokens() con
 	return NamingTokens;
 }
 
+void ULiveLinkHubUserSettings::PostInitProperties()
+{
+	Super::PostInitProperties();
+
+	// Cache the filter data on load.
+	UE::TWriteScopeLock Lock{ FiltersLock };
+	CurrentFilterPreset_ThreadSafe = CurrentPreset;
+}
+
 void ULiveLinkHubUserSettings::CacheRecentConfig(const FString& SavePath)
 {
 	LastConfigDirectory = FPaths::GetPath(SavePath);
@@ -91,3 +100,105 @@ TArray<FLiveLinkHubSessionFile> ULiveLinkHubUserSettings::GetRecentConfigFiles()
 
 	return Files;
 }
+
+void ULiveLinkHubUserSettings::SaveCurrentPreset()
+{
+	for (FLiveLinkHubClientFilterPreset& Preset : FilterPresets)
+	{
+		if (Preset.PresetName == CurrentPreset.PresetName)
+		{
+			Preset = CurrentPreset;
+			break;
+		}
+	}
+
+	SaveConfig();
+}
+
+void ULiveLinkHubUserSettings::SavePresetAs(const FString& PresetName)
+{
+	FilterPresets.Add_GetRef(CurrentPreset).PresetName = PresetName;
+	CurrentPreset.PresetName = PresetName;
+
+	{
+		UE::TWriteScopeLock Lock{ FiltersLock };
+		CurrentFilterPreset_ThreadSafe = CurrentPreset;
+	}
+
+	SaveConfig();
+}
+
+void ULiveLinkHubUserSettings::SaveClientsToPreset(const FString& InPresetName)
+{
+	TArray<TPair<FString, FString>> IPAndProject;
+	if (ILiveLinkHubClientsModel* ClientsModel = ILiveLinkHubClientsModel::Get())
+	{
+		for (const FLiveLinkHubClientId& ClientId : ClientsModel->GetSessionClients())
+		{
+			if (TOptional<FLiveLinkHubUEClientInfo> Info = ClientsModel->GetClientInfo(ClientId))
+			{
+				if (!Info->IPAddress.IsEmpty())
+				{
+					IPAndProject.Emplace(Info->IPAddress, Info->ProjectName);
+				}
+			}
+		}
+	}
+
+	FLiveLinkHubClientFilterPreset Preset;
+	Preset.PresetName = InPresetName;
+	Preset.Filters.Reserve(IPAndProject.Num());
+	Preset.AutoConnectClients = CurrentPreset.AutoConnectClients;
+
+	Algo::Transform(IPAndProject, Preset.Filters, [](const TPair<FString, FString>& IpAndProject)
+		{
+			return FLiveLinkHubClientTextFilter{ IpAndProject.Value, IpAndProject.Key, ELiveLinkHubClientFilterType::IP, ELiveLinkHubClientFilterBehavior::Include };
+		});
+
+	int32 Index = FilterPresets.IndexOfByPredicate([&InPresetName](const FLiveLinkHubClientFilterPreset& Preset) { return InPresetName == Preset.PresetName; });
+	if (Index != INDEX_NONE)
+	{
+		FilterPresets[Index] = Preset;
+	}
+	else
+	{
+		FilterPresets.Add(Preset);
+	}
+
+	CurrentPreset = Preset;
+
+	SaveConfig();
+
+	UE::TWriteScopeLock Lock{ FiltersLock };
+	CurrentFilterPreset_ThreadSafe = CurrentPreset;
+}
+
+void ULiveLinkHubUserSettings::LoadFilterPreset(const FString& PresetName)
+{
+	if (FLiveLinkHubClientFilterPreset* Preset = Algo::FindBy(FilterPresets, PresetName, &FLiveLinkHubClientFilterPreset::PresetName))
+	{
+		{
+			UE::TWriteScopeLock Lock{ FiltersLock };
+			CurrentPreset = *Preset;
+		}
+
+		PostUpdateClientFilters();
+	}
+	else
+	{
+		UE_LOG(LogLiveLinkHub, Warning, TEXT("Could not find preset %s."), *PresetName);
+	}
+}
+
+void ULiveLinkHubUserSettings::DeleteFilterPreset(const FString& PresetName)
+{
+	FilterPresets.RemoveAll([&PresetName](const FLiveLinkHubClientFilterPreset& InPreset) { return PresetName == InPreset.PresetName; });
+	if (PresetName == CurrentPreset.PresetName)
+	{
+		CurrentPreset = FLiveLinkHubClientFilterPreset{};
+
+		UE::TWriteScopeLock Lock{ FiltersLock };
+		CurrentFilterPreset_ThreadSafe = CurrentPreset;
+	}
+	SaveConfig();
+}
diff --git a/Engine/Plugins/Animation/LiveLinkHub/Source/LiveLinkHub/Private/Settings/LiveLinkHubSettings.h b/Engine/Plugins/Animation/LiveLinkHub/Source/LiveLinkHub/Private/Settings/LiveLinkHubSettings.h
--- a/Engine/Plugins/Animation/LiveLinkHub/Source/LiveLinkHub/Private/Settings/LiveLinkHubSettings.h
+++ b/Engine/Plugins/Animation/LiveLinkHub/Source/LiveLinkHub/Private/Settings/LiveLinkHubSettings.h
@@ -211,15 +211,8 @@ class ULiveLinkHubUserSettings : public UObject
 	GENERATED_BODY()
 
 public:
-
-	virtual void PostInitProperties() override
-	{
-		Super::PostInitProperties();
-
-		// Cache the filter data on load.
-		UE::TWriteScopeLock Lock{ FiltersLock };
-		CurrentFilterPreset_ThreadSafe = CurrentPreset;
-	}
+	//~ UObject Interface
+	virtual void PostInitProperties() override;
 
 	/** Number of recent configs to keep track. */
 	UPROPERTY(config, EditAnywhere, Category = "LiveLinkHub")
@@ -282,51 +275,14 @@ public:
 		PostUpdateClientFilters();
 	}
 
-	/** Creates a preset of filters from the list of clients that are currently in the session. */
-	void SaveClientsToPreset(const FString& InPresetName)
-	{
-		TArray<TPair<FString, FString>> IPAndProject;
-		if (ILiveLinkHubClientsModel* ClientsModel = ILiveLinkHubClientsModel::Get())
-		{
-			for (const FLiveLinkHubClientId& ClientId : ClientsModel->GetSessionClients())
-			{
-				if (TOptional<FLiveLinkHubUEClientInfo> Info = ClientsModel->GetClientInfo(ClientId))
-				{
-					if (!Info->IPAddress.IsEmpty())
-					{
-						IPAndProject.Emplace(Info->IPAddress, Info->ProjectName);
-					}
-				}
-			}
-		}
-
-		FLiveLinkHubClientFilterPreset Preset;
-		Preset.PresetName = InPresetName;
-		Preset.Filters.Reserve(IPAndProject.Num());
-		Preset.AutoConnectClients = CurrentPreset.AutoConnectClients;
-
-		Algo::Transform(IPAndProject, Preset.Filters, [](const TPair<FString, FString>& IpAndProject)
-		{
-			return FLiveLinkHubClientTextFilter{ IpAndProject.Value, IpAndProject.Key, ELiveLinkHubClientFilterType::IP, ELiveLinkHubClientFilterBehavior::Include };
-		});
-
-		int32 Index = FilterPresets.IndexOfByPredicate([&InPresetName](const FLiveLinkHubClientFilterPreset& Preset) { return InPresetName == Preset.PresetName; });
-		if (Index != INDEX_NONE)
-		{
-			FilterPresets[Index] = Preset;
-		}
-		else
-		{
-			FilterPresets.Add(Preset);
-		}
-
-		CurrentPreset = Preset;
+	/** Save the current preset to the settings ini. */
+	void SaveCurrentPreset();
 
-		SaveConfig();
+	/** Save the current filters configuration as a new preset in the settings ini. */
+	void SavePresetAs(const FString& PresetName);
 
-		UE::TWriteScopeLock Lock{ FiltersLock };
-		CurrentFilterPreset_ThreadSafe = CurrentPreset;
-	}
+	/** Creates a preset of filters from the list of clients that are currently in the session. */
+	void SaveClientsToPreset(const FString& InPresetName);
 
 	/** Get all filter presets. */
 	const TArray<FLiveLinkHubClientFilterPreset>& GetFilterPresets() const
@@ -335,33 +291,10 @@ public:
 	}
 
 	/** Applies a filter preset. */
-	void LoadFilterPreset(const FString& PresetName)
-	{
-		if (FLiveLinkHubClientFilterPreset* Preset = Algo::FindBy(FilterPresets, PresetName, &FLiveLinkHubClientFilterPreset::PresetName))
-		{
-			{
-				UE::TWriteScopeLock Lock{ FiltersLock };
-				CurrentPreset = *Preset;
-			}
-
-			PostUpdateClientFilters();
-		}
-		else
-		{
-			UE_LOG(LogLiveLinkHub, Warning, TEXT("Could not find preset %s."), *PresetName);
-		}
-	}
+	void LoadFilterPreset(const FString& PresetName);
 
 	/** Deletes a filter preset. */
-	void DeleteFilterPreset(const FString& PresetName)
-	{
-		FilterPresets.RemoveAll([&PresetName](const FLiveLinkHubClientFilterPreset& InPreset) { return PresetName == InPreset.PresetName; });
-		if (PresetName == CurrentPreset.PresetName)
-		{
-			CurrentPreset = FLiveLinkHubClientFilterPreset{};
-		}
-		SaveConfig();
-	}
+	void DeleteFilterPreset(const FString& PresetName);
 
 	/** Get a copy of the client filters data in a thread-safe way. */
 	FLiveLinkHubClientFilterPreset GetClientFiltersData_AnyThread()
