LiveLinkHub - Add UdpMessaging plugin dependency to LiveLinkHub

#rnx

diff --git a/Engine/Plugins/Animation/LiveLinkHub/LiveLinkHub.uplugin b/Engine/Plugins/Animation/LiveLinkHub/LiveLinkHub.uplugin
--- a/Engine/Plugins/Animation/LiveLinkHub/LiveLinkHub.uplugin
+++ b/Engine/Plugins/Animation/LiveLinkHub/LiveLinkHub.uplugin
@@ -24,6 +24,10 @@
 		{
 			"Name": "NamingTokens",
 			"Enabled": true
+		},
+		{
+			"Name": "UdpMessaging",
+			"Enabled": true
 		}
 	],
 	"Modules": [
