LiveLinkHub - Disable timecode lerp interpolation when in playback

diff --git a/Engine/Plugins/Animation/LiveLinkHub/Source/LiveLinkHub/Private/LiveLinkHubModule.cpp b/Engine/Plugins/Animation/LiveLinkHub/Source/LiveLinkHub/Private/LiveLinkHubModule.cpp
--- a/Engine/Plugins/Animation/LiveLinkHub/Source/LiveLinkHub/Private/LiveLinkHubModule.cpp
+++ b/Engine/Plugins/Animation/LiveLinkHub/Source/LiveLinkHub/Private/LiveLinkHubModule.cpp
@@ -18,6 +18,7 @@
 #include "Settings/LiveLinkHubSettings.h"
 #include "Settings/LiveLinkHubClientTextFilterCustomization.h"
 #include "Settings/LiveLinkHubSettingsCustomization.h"
+#include "Settings/LiveLinkHubTimecodeSettingsCustomization.h"
 #include "Settings/LiveLinkHubCustomTimeStepSettingsCustomization.h"
 #include "Settings/LiveLinkSettingsCustomization.h"
 #include "Textures/SlateIcon.h"
@@ -55,6 +56,7 @@ void FLiveLinkHubModule::StartupModule()
 	PropertyModule.RegisterCustomClassLayout(ULiveLinkHubSettings::StaticClass()->GetFName(), FOnGetDetailCustomizationInstance::CreateStatic(&FLiveLinkHubSettingsCustomization::MakeInstance));
 	PropertyModule.RegisterCustomPropertyTypeLayout(FLiveLinkHubClientTextFilter::StaticStruct()->GetFName(), FOnGetPropertyTypeCustomizationInstance::CreateStatic(&FLiveLinkHubClientTextFilterCustomization::MakeInstance));
 	PropertyModule.RegisterCustomPropertyTypeLayout(FLiveLinkHubCustomTimeStepSettings::StaticStruct()->GetFName(), FOnGetPropertyTypeCustomizationInstance::CreateStatic(&FLiveLinkHubCustomTimeStepSettingsCustomization::MakeInstance));
+	PropertyModule.RegisterCustomPropertyTypeLayout(FLiveLinkHubTimecodeSettings::StaticStruct()->GetFName(), FOnGetPropertyTypeCustomizationInstance::CreateStatic(&FLiveLinkHubTimecodeSettingsCustomization::MakeInstance));
 
 	// Apply our customization for core live link settings, only if we aren't running in the full editor. We hide properties that aren't
 	// supported in a standalone application context, but are needed if loaded in the editor.
@@ -74,6 +76,7 @@ void FLiveLinkHubModule::ShutdownModule()
 		PropertyEditorModule->UnregisterCustomClassLayout(ULiveLinkHubSettings::StaticClass()->GetFName());
 		if (bUseSettingsDetailCustomization)
 		{
+			PropertyEditorModule->UnregisterCustomClassLayout(FLiveLinkHubTimecodeSettings::StaticStruct()->GetFName());
 			PropertyEditorModule->UnregisterCustomClassLayout(FLiveLinkHubCustomTimeStepSettings::StaticStruct()->GetFName());
 			PropertyEditorModule->UnregisterCustomClassLayout(ULiveLinkSettings::StaticClass()->GetFName());
 		}
diff --git a/Engine/Plugins/Animation/LiveLinkHub/Source/LiveLinkHub/Private/Recording/LiveLinkHubPlaybackController.cpp b/Engine/Plugins/Animation/LiveLinkHub/Source/LiveLinkHub/Private/Recording/LiveLinkHubPlaybackController.cpp
--- a/Engine/Plugins/Animation/LiveLinkHub/Source/LiveLinkHub/Private/Recording/LiveLinkHubPlaybackController.cpp
+++ b/Engine/Plugins/Animation/LiveLinkHub/Source/LiveLinkHub/Private/Recording/LiveLinkHubPlaybackController.cpp
@@ -32,11 +32,26 @@ namespace UE::LiveLinkHub
 {
 	static void ClearTimecodeProviderFrameHistory()
 	{
-		if (ULiveLinkTimecodeProvider* Provider = Cast< ULiveLinkTimecodeProvider>(GEngine->GetTimecodeProvider()))
+		if (GEngine)
 		{
-			Provider->ClearFrameHistory();
+			if (ULiveLinkTimecodeProvider* Provider = Cast< ULiveLinkTimecodeProvider>(GEngine->GetTimecodeProvider()))
+			{
+				Provider->ClearFrameHistory();
+			}
 		}
 	}
+
+	static void OverrideTimecodeEvaluationType()
+	{
+		GetMutableDefault<ULiveLinkHubTimeAndSyncSettings>()->TimecodeSettings.OverrideEvaluationType = ELiveLinkTimecodeProviderEvaluationType::Latest;
+		GetMutableDefault<ULiveLinkHubTimeAndSyncSettings>()->OnToggleTimecodeSettings();
+	}
+
+	static void ResetTimecodeEvaluationType()
+	{
+		GetMutableDefault<ULiveLinkHubTimeAndSyncSettings>()->TimecodeSettings.OverrideEvaluationType.Reset();
+		GetMutableDefault<ULiveLinkHubTimeAndSyncSettings>()->OnToggleTimecodeSettings();
+	}
 }
 
 /**
@@ -320,6 +335,8 @@ void FLiveLinkHubPlaybackController::PreparePlayback(ULiveLinkRecording* InLiveL
 
 				RecordingToPlay->RecordingPreset->ApplyToClientLatent([this](bool)
 				{
+					UE::LiveLinkHub::OverrideTimecodeEvaluationType();
+
 					bIsReady = true;
 					SyncToFrame(FQualifiedFrameTime()); // Needed to establish connection with client
 					
@@ -513,6 +530,8 @@ void FLiveLinkHubPlaybackController::EjectAndUnload(const TFunction<void()>& Eje
 		constexpr bool bUnloadNextTick = true;
 		UnloadRecordingPackage(Package, bUnloadNextTick);
 	}
+
+	UE::LiveLinkHub::ResetTimecodeEvaluationType();
 }
 
 void FLiveLinkHubPlaybackController::UnloadRecordingPackage(const TWeakObjectPtr<UPackage>& InPackage, bool bUnloadNextTick)
diff --git a/Engine/Plugins/Animation/LiveLinkHub/Source/LiveLinkHub/Private/Settings/LiveLinkHubTimecodeSettingsCustomization.h b/Engine/Plugins/Animation/LiveLinkHub/Source/LiveLinkHub/Private/Settings/LiveLinkHubTimecodeSettingsCustomization.h
new file mode 100644
--- /dev/null
+++ b/Engine/Plugins/Animation/LiveLinkHub/Source/LiveLinkHub/Private/Settings/LiveLinkHubTimecodeSettingsCustomization.h
@@ -0,0 +1,89 @@
+// Copyright Epic Games, Inc. All Rights Reserved.
+
+#pragma once
+
+#include "IPropertyTypeCustomization.h"
+
+#include "DetailWidgetRow.h"
+#include "IDetailChildrenBuilder.h"
+#include "IDetailPropertyRow.h"
+#include "LiveLinkHubMessages.h"
+#include "PropertyHandle.h"
+#include "Widgets/Text/STextBlock.h"
+
+
+#define LOCTEXT_NAMESPACE "LiveLinkHubTimecodeSettingsCustomization"
+
+/**
+ * Customization for the Timecode settings.
+ */
+class FLiveLinkHubTimecodeSettingsCustomization : public IPropertyTypeCustomization
+{
+public:
+	static TSharedRef<IPropertyTypeCustomization> MakeInstance()
+	{
+		return MakeShared<FLiveLinkHubTimecodeSettingsCustomization>();
+	}
+
+	virtual void CustomizeHeader(TSharedRef<IPropertyHandle> PropertyHandle, FDetailWidgetRow& HeaderRow, IPropertyTypeCustomizationUtils& CustomizationUtils)
+	{
+	}
+
+	virtual void CustomizeChildren(TSharedRef<IPropertyHandle> PropertyHandle, IDetailChildrenBuilder& ChildBuilder, IPropertyTypeCustomizationUtils& CustomizationUtils) override
+	{
+		uint32 NumberOfChild;
+		if (PropertyHandle->GetNumChildren(NumberOfChild) == FPropertyAccess::Success)
+		{
+			for (uint32 Index = 0; Index < NumberOfChild; ++Index)
+			{
+				TSharedRef<IPropertyHandle> ChildPropertyHandle = PropertyHandle->GetChildHandle(Index).ToSharedRef();
+
+				if (ChildPropertyHandle->GetProperty()->GetFName() == GET_MEMBER_NAME_CHECKED(FLiveLinkHubTimecodeSettings, EvaluationType))
+				{
+					TSharedPtr<IPropertyHandle> OverrideHandle = PropertyHandle->GetChildHandle(GET_MEMBER_NAME_CHECKED(FLiveLinkHubTimecodeSettings, OverrideEvaluationType));
+
+					TArray<UObject*> Objects;
+					PropertyHandle->GetOuterObjects(Objects);
+
+					if (Objects.Num() == 0)
+					{
+						continue;
+					}
+
+					uint8* BaseAddress = PropertyHandle->GetValueBaseAddress((uint8*)Objects[0]);
+					FProperty* Property = FLiveLinkHubTimecodeSettings::StaticStruct()->FindPropertyByName(GET_MEMBER_NAME_CHECKED(FLiveLinkHubTimecodeSettings, OverrideEvaluationType));
+
+					TOptional<ELiveLinkTimecodeProviderEvaluationType> OptionalOverride;
+					Property->GetValue_InContainer(BaseAddress, &OptionalOverride);
+						
+					if (OptionalOverride.IsSet())
+					{
+						// Override the EvaluationType dropdown to display the override value in a read-only widget.
+						FText OverrideText = StaticEnum<ELiveLinkTimecodeProviderEvaluationType>()->GetDisplayNameTextByValue((int64)OptionalOverride.GetValue());
+
+						ChildBuilder.AddCustomRow(ChildPropertyHandle->GetPropertyDisplayName())
+						.NameContent()
+						[
+							ChildPropertyHandle->CreatePropertyNameWidget()
+						]
+						.ValueContent()
+						[
+							SNew(STextBlock)
+								.Text(FText::Format(LOCTEXT("EvaluationTypeOverriddenProperty", "{0} (Overridden in playback)"), OverrideText))
+						];
+					}
+					else
+					{
+						ChildBuilder.AddProperty(ChildPropertyHandle);
+					}
+				}
+				else
+				{
+					ChildBuilder.AddProperty(ChildPropertyHandle);
+				}
+			}
+		}
+	}
+};
+
+#undef LOCTEXT_NAMESPACE
diff --git a/Engine/Plugins/Animation/LiveLinkHub/Source/LiveLinkHubMessaging/Private/LiveLinkHubMessages.cpp b/Engine/Plugins/Animation/LiveLinkHub/Source/LiveLinkHubMessaging/Private/LiveLinkHubMessages.cpp
--- a/Engine/Plugins/Animation/LiveLinkHub/Source/LiveLinkHubMessaging/Private/LiveLinkHubMessages.cpp
+++ b/Engine/Plugins/Animation/LiveLinkHub/Source/LiveLinkHubMessaging/Private/LiveLinkHubMessages.cpp
@@ -130,7 +130,7 @@ void FLiveLinkHubTimecodeSettings::AssignTimecodeSettingsAsProviderToEngine() co
 		LiveLinkProvider->OverrideFrameRate = DesiredFrameRate;
 		LiveLinkProvider->FrameDelay = FrameDelay;
 		LiveLinkProvider->BufferSize = BufferSize;
-		LiveLinkProvider->Evaluation = EvaluationType;
+		LiveLinkProvider->Evaluation = OverrideEvaluationType ? *OverrideEvaluationType : EvaluationType;
 		GEngine->SetTimecodeProvider(LiveLinkProvider);
 		UE_LOG(LogLiveLinkHubMessages, Display, TEXT("Live Link Timecode provider assigned to %s."), *SubjectName.ToString());
 	}
diff --git a/Engine/Plugins/Animation/LiveLinkHub/Source/LiveLinkHubMessaging/Public/LiveLinkHubMessages.h b/Engine/Plugins/Animation/LiveLinkHub/Source/LiveLinkHubMessaging/Public/LiveLinkHubMessages.h
--- a/Engine/Plugins/Animation/LiveLinkHub/Source/LiveLinkHubMessaging/Public/LiveLinkHubMessages.h
+++ b/Engine/Plugins/Animation/LiveLinkHub/Source/LiveLinkHubMessaging/Public/LiveLinkHubMessages.h
@@ -127,9 +127,16 @@ struct FLiveLinkHubTimecodeSettings
 	UPROPERTY(config, EditAnywhere, Category = "Timecode", meta = (ClampMin = "2", UIMin = "2", ClampMax = "10", UIMax = "10"))
 	int32 BufferSize = 2;
 
-	/** How timecode should be evaluated. */
+	/** 
+	 * How timecode should be evaluated. 
+	 * @Note: This will be overridden while a Live Link Hub recording is playing.
+	 */
 	UPROPERTY(config, EditAnywhere, Category = "Timecode")
-	ELiveLinkTimecodeProviderEvaluationType EvaluationType = ELiveLinkTimecodeProviderEvaluationType::Lerp;
+	ELiveLinkTimecodeProviderEvaluationType EvaluationType = ELiveLinkTimecodeProviderEvaluationType::Latest;
+
+	/** If set, this will override the timecode evaluation type. */
+	UPROPERTY(Transient)
+	TOptional<ELiveLinkTimecodeProviderEvaluationType> OverrideEvaluationType;
 
 	/** Assign the settings to a new timecode provider and override the current engine settings. */
 	UE_API void AssignTimecodeSettingsAsProviderToEngine() const;
