LiveLinkHub - Fix crash recovery prompt not being sent to the foreground

#rnx

diff --git a/Engine/Plugins/Animation/LiveLinkHub/Source/LiveLinkHub/Private/LiveLinkHubCrashRecovery.h b/Engine/Plugins/Animation/LiveLinkHub/Source/LiveLinkHub/Private/LiveLinkHubCrashRecovery.h
--- a/Engine/Plugins/Animation/LiveLinkHub/Source/LiveLinkHub/Private/LiveLinkHubCrashRecovery.h
+++ b/Engine/Plugins/Animation/LiveLinkHub/Source/LiveLinkHub/Private/LiveLinkHubCrashRecovery.h
@@ -119,8 +119,19 @@ private:
 	{
 		if (TSharedPtr<SWindow> ParentWindow = FGlobalTabmanager::Get()->GetRootWindow())
 		{
-			ParentWindow->BringToFront();
-			ParentWindow->FlashWindow();
+			/**
+			 * Not sure if it's a combination of it being minimized + startup screen but I could not get the Recovery window to show up without HACK_ForceToFront.
+			 * Things that did not work:
+			 * - Adding a timer for tick and calling BringToFront there
+			 * - Calling Maximize on the root window
+			 * - Calling BringToFront with bForce=true.
+			 * 
+			 * So I instead copied the logic from CrashReportClient.
+			 */
+			const bool bForceBringToFront = (false || (PLATFORM_MAC));
+
+			ParentWindow->HACK_ForceToFront();
+			ParentWindow->BringToFront(bForceBringToFront);
 		}
 
 		const EAppReturnType::Type Response = FMessageDialog::Open(
