LiveLinkHub - Fix crash when trying to save over old recording in Standalone version of LiveLinkHub

#rnx

diff --git a/Engine/Plugins/Animation/LiveLinkHub/Source/LiveLinkHub/Private/Recording/Implementations/LiveLinkUAssetRecorder.cpp b/Engine/Plugins/Animation/LiveLinkHub/Source/LiveLinkHub/Private/Recording/Implementations/LiveLinkUAssetRecorder.cpp
--- a/Engine/Plugins/Animation/LiveLinkHub/Source/LiveLinkHub/Private/Recording/Implementations/LiveLinkUAssetRecorder.cpp
+++ b/Engine/Plugins/Animation/LiveLinkHub/Source/LiveLinkHub/Private/Recording/Implementations/LiveLinkUAssetRecorder.cpp
@@ -187,7 +187,9 @@ bool FLiveLinkUAssetRecorder::GetSavePresetPackageName(FString& OutName, const F
 		}
 	}
 	
-	const FString ContentDir = FPaths::ProjectContentDir();
+	FString ContentDir;
+
+	ensure(FPackageName::TryConvertLongPackageNameToFilename(FString(TEXT("/Game")), ContentDir));
 	const FString DialogStartPath = FPaths::Combine(TEXT("/Game"), DefaultFolder);
 	const FString AbsoluteFolderPath = FPaths::Combine(ContentDir, DefaultFolder);
 
