LiveLinkHub - Fix virtual subject name not persisting after config reload

#rnx

#DisabledValidations Preflight

diff --git a/Engine/Plugins/Animation/LiveLinkHub/Source/LiveLinkHub/Private/LiveLinkHubAnimationVirtualSubject.h b/Engine/Plugins/Animation/LiveLinkHub/Source/LiveLinkHub/Private/LiveLinkHubAnimationVirtualSubject.h
--- a/Engine/Plugins/Animation/LiveLinkHub/Source/LiveLinkHub/Private/LiveLinkHubAnimationVirtualSubject.h
+++ b/Engine/Plugins/Animation/LiveLinkHub/Source/LiveLinkHub/Private/LiveLinkHubAnimationVirtualSubject.h
@@ -26,7 +26,11 @@ public:
 	virtual void Initialize(FLiveLinkSubjectKey InSubjectKey, TSubclassOf<ULiveLinkRole> InRole, ILiveLinkClient* InLiveLinkClient) override
 	{
 		ULiveLinkAnimationVirtualSubject::Initialize(InSubjectKey, InRole, InLiveLinkClient);
-		OutboundName = InSubjectKey.SubjectName.ToString();
+
+		if (OutboundName.IsEmpty())
+		{
+			OutboundName = InSubjectKey.SubjectName.ToString();
+		}
 
 		FLiveLinkClient* Client = &IModularFeatures::Get().GetModularFeature<FLiveLinkClient>(FLiveLinkClient::ModularFeatureName);
 		Source = Client->GetSourceType(InSubjectKey.Source).ToString();
