LiveLinkHub - Add button to open content folder and to refresh recordings when the recordings list is empty.

#DisabledValidations Preflight

diff --git a/Engine/Plugins/Animation/LiveLinkHub/Source/LiveLinkHub/Private/Recording/SLiveLinkHubRecordingListView.h b/Engine/Plugins/Animation/LiveLinkHub/Source/LiveLinkHub/Private/Recording/SLiveLinkHubRecordingListView.h
--- a/Engine/Plugins/Animation/LiveLinkHub/Source/LiveLinkHub/Private/Recording/SLiveLinkHubRecordingListView.h
+++ b/Engine/Plugins/Animation/LiveLinkHub/Source/LiveLinkHub/Private/Recording/SLiveLinkHubRecordingListView.h
@@ -19,10 +19,10 @@
 #include "Recording/LiveLinkRecording.h"
 #include "UObject/SavePackage.h"
 #include "Widgets/Input/SButton.h"
+#include "Widgets/Layout/SWidgetSwitcher.h"
 #include "Widgets/SCompoundWidget.h"
 #include "Widgets/Text/STextBlock.h"
 
-
 #define LOCTEXT_NAMESPACE "LiveLinkHub.RecordingListView"
 
 class SLiveLinkHubRecordingListView : public SCompoundWidget
@@ -65,28 +65,60 @@ public:
 		
 		ChildSlot
 		[
-			SNew(SVerticalBox)
-			+ SVerticalBox::Slot()
+			SNew(SWidgetSwitcher)
+			.WidgetIndex_Lambda([this]() { return GetRecordingPickerVisibility() == EVisibility::Visible ? 0 : 1; })
+			+ SWidgetSwitcher::Slot()
 			.VAlign(VAlign_Fill)
 			[
 				SAssignNew(BoxWidget, SBox)
-				.Visibility(this, &SLiveLinkHubRecordingListView::GetRecordingPickerVisibility)
 				[
 					CreateRecordingPicker()
 				]
 			]
-			+SVerticalBox::Slot()
-			.AutoHeight()
-			.VAlign(VAlign_Top)
+			+ SWidgetSwitcher::Slot()
+			.VAlign(VAlign_Center)
 			.HAlign(HAlign_Center)
 			[
-				SNew(STextBlock)
-				.Visibility_Lambda([this]()
-				{
-					EVisibility RecordingPickerVisibility = GetRecordingPickerVisibility();
-					return RecordingPickerVisibility == EVisibility::Visible ? EVisibility::Collapsed : EVisibility::Visible;
-				})
-				.Text(GetNoAssetsWarningText())
+				SNew(SVerticalBox)
+				+ SVerticalBox::Slot()
+				.HAlign(HAlign_Center)
+				.AutoHeight()
+				[
+					SNew(STextBlock)
+					.Text(GetNoAssetsWarningText())
+				]
+				+ SVerticalBox::Slot()
+				.Padding(FMargin(4, 4))
+				.HAlign(HAlign_Center)
+				.AutoHeight()
+				[
+					SNew(SHorizontalBox)
+					+ SHorizontalBox::Slot()
+					.Padding(FMargin(4.0, 2.0))
+					.AutoWidth()
+					[
+						SNew(SButton)
+						.Text(LOCTEXT("RecordingsFolderButton", "Open Recordings Folder"))
+						.OnClicked_Lambda([]()
+						{
+							FPlatformProcess::ExploreFolder(*FPaths::ConvertRelativePathToFull(FPaths::ProjectContentDir()));
+							return FReply::Handled();
+						})
+					]
+					+ SHorizontalBox::Slot()
+					.Padding(FMargin(4.0, 2.0))
+					.AutoWidth()
+					[
+						SNew(SButton)
+						.Text(LOCTEXT("RefreshRecordingsButton", "Refresh Recordings"))
+						.OnClicked_Lambda([]()
+						{
+							const FAssetRegistryModule& AssetRegistryModule = FModuleManager::LoadModuleChecked<FAssetRegistryModule>("AssetRegistry");
+							AssetRegistryModule.Get().ScanPathsSynchronous({ TEXT("/Game") }, true);
+							return FReply::Handled();
+						})
+					]
+				]
 			]
 		];
 	}
@@ -391,7 +423,7 @@ private:
 	/** The text to display when no assets are found. */
 	static FText GetNoAssetsWarningText()
 	{
-		return LOCTEXT("NoRecordings_Warning", "No Recordings Found");
+		return LOCTEXT("NoRecordings_Warning", "No recordings found. Press Record to create a new recording.");
 	}
 	
 private:
