[Backout] - CL47324281
#rnx
Original CL Desc
-----------------------------------------------------------------
Port MessageBus tester plugin into the UE experimental folder. This plugin is designed to provide test cases that can be used to exercise UDP message transport. The test cases help analyze use cases for messaging and make sure the low level code is functioning optimally.  This is only designed for testing and debugging at the moment and may be changed or removed in the future.

diff --git a/Engine/Plugins/Animation/LiveLinkHub/LiveLinkHub.uplugin b/Engine/Plugins/Animation/LiveLinkHub/LiveLinkHub.uplugin
--- a/Engine/Plugins/Animation/LiveLinkHub/LiveLinkHub.uplugin
+++ b/Engine/Plugins/Animation/LiveLinkHub/LiveLinkHub.uplugin
@@ -28,10 +28,6 @@
 		{
 			"Name": "UdpMessaging",
 			"Enabled": true
-		},
-        {
-			"Name": "MessageBusTester",
-			"Enabled": true
 		}
 	],
 	"Modules": [
diff --git a/Engine/Plugins/Animation/LiveLinkHub/Source/LiveLinkHub/LiveLinkHub.Build.cs b/Engine/Plugins/Animation/LiveLinkHub/Source/LiveLinkHub/LiveLinkHub.Build.cs
--- a/Engine/Plugins/Animation/LiveLinkHub/Source/LiveLinkHub/LiveLinkHub.Build.cs
+++ b/Engine/Plugins/Animation/LiveLinkHub/Source/LiveLinkHub/LiveLinkHub.Build.cs
@@ -48,7 +48,6 @@ public class LiveLinkHub : ModuleRules
 				"UnrealEd",
 				"OutputLog",
 				"NamingTokens",
-				"MessageBusTesterEditor"
 			});
 
 		if (Target.Configuration != UnrealTargetConfiguration.Shipping)
diff --git a/Engine/Plugins/Animation/LiveLinkHub/Source/LiveLinkHub/Private/LiveLinkHubApplication.cpp b/Engine/Plugins/Animation/LiveLinkHub/Source/LiveLinkHub/Private/LiveLinkHubApplication.cpp
--- a/Engine/Plugins/Animation/LiveLinkHub/Source/LiveLinkHub/Private/LiveLinkHubApplication.cpp
+++ b/Engine/Plugins/Animation/LiveLinkHub/Source/LiveLinkHub/Private/LiveLinkHubApplication.cpp
@@ -24,7 +24,6 @@
 #		include "ISessionFrontendModule.h"
 #	endif
 #	include "ISlateReflectorModule.h"
-#	include "IMessageBusTesterEditorModule.h"
 #endif
 
 #define LOCTEXT_NAMESPACE "LiveLinkHubApplication"
@@ -582,21 +581,8 @@ void FLiveLinkHubApplication::CreateToolsMenu(FMenuBuilder& MenuBuilder, UToolMe
 				SlateReflector.DisplayWidgetReflector();
 			})
 	);
-	MenuBuilder.AddMenuEntry(
-		LOCTEXT("FileMenu_Development_MesageBusTester", "Message Bus Tester"),
-		FText(),
-		FSlateIcon(),
-		FExecuteAction::CreateStatic([]()
-			{
-				if (IMessageBusTesterEditorModule::IsAvailable())
-				{
-					IMessageBusTesterEditorModule::Get().DisplayMessageBusTester();
-				}
-			})
-	);
 	MenuBuilder.EndSection();
 
-
 #endif /*UE_BUILD_SHIPPING*/
 }
 
diff --git a/Engine/Plugins/Experimental/MessageBusTester/MessageBusTester.uplugin b/Engine/Plugins/Experimental/MessageBusTester/MessageBusTester.uplugin
deleted file mode 100644
diff --git a/Engine/Plugins/Experimental/MessageBusTester/Source/MessageBusTester/MessageBusTester.Build.cs b/Engine/Plugins/Experimental/MessageBusTester/Source/MessageBusTester/MessageBusTester.Build.cs
deleted file mode 100644
diff --git a/Engine/Plugins/Experimental/MessageBusTester/Source/MessageBusTester/Private/DiscoveredTester.cpp b/Engine/Plugins/Experimental/MessageBusTester/Source/MessageBusTester/Private/DiscoveredTester.cpp
deleted file mode 100644
diff --git a/Engine/Plugins/Experimental/MessageBusTester/Source/MessageBusTester/Private/MessageBusTester.cpp b/Engine/Plugins/Experimental/MessageBusTester/Source/MessageBusTester/Private/MessageBusTester.cpp
deleted file mode 100644
diff --git a/Engine/Plugins/Experimental/MessageBusTester/Source/MessageBusTester/Private/MessageBusTester.h b/Engine/Plugins/Experimental/MessageBusTester/Source/MessageBusTester/Private/MessageBusTester.h
deleted file mode 100644
diff --git a/Engine/Plugins/Experimental/MessageBusTester/Source/MessageBusTester/Private/MessageBusTesterCommon.cpp b/Engine/Plugins/Experimental/MessageBusTester/Source/MessageBusTester/Private/MessageBusTesterCommon.cpp
deleted file mode 100644
diff --git a/Engine/Plugins/Experimental/MessageBusTester/Source/MessageBusTester/Private/MessageBusTesterLogger.cpp b/Engine/Plugins/Experimental/MessageBusTester/Source/MessageBusTester/Private/MessageBusTesterLogger.cpp
deleted file mode 100644
diff --git a/Engine/Plugins/Experimental/MessageBusTester/Source/MessageBusTester/Private/MessageBusTesterLogger.h b/Engine/Plugins/Experimental/MessageBusTester/Source/MessageBusTester/Private/MessageBusTesterLogger.h
deleted file mode 100644
diff --git a/Engine/Plugins/Experimental/MessageBusTester/Source/MessageBusTester/Private/MessageBusTesterModule.cpp b/Engine/Plugins/Experimental/MessageBusTester/Source/MessageBusTester/Private/MessageBusTesterModule.cpp
deleted file mode 100644
diff --git a/Engine/Plugins/Experimental/MessageBusTester/Source/MessageBusTester/Private/MessageBusTesterModule.h b/Engine/Plugins/Experimental/MessageBusTester/Source/MessageBusTester/Private/MessageBusTesterModule.h
deleted file mode 100644
diff --git a/Engine/Plugins/Experimental/MessageBusTester/Source/MessageBusTester/Private/MessageBusTesterSettings.cpp b/Engine/Plugins/Experimental/MessageBusTester/Source/MessageBusTester/Private/MessageBusTesterSettings.cpp
deleted file mode 100644
diff --git a/Engine/Plugins/Experimental/MessageBusTester/Source/MessageBusTester/Public/DiscoveredTester.h b/Engine/Plugins/Experimental/MessageBusTester/Source/MessageBusTester/Public/DiscoveredTester.h
deleted file mode 100644
diff --git a/Engine/Plugins/Experimental/MessageBusTester/Source/MessageBusTester/Public/IMessageBusTester.h b/Engine/Plugins/Experimental/MessageBusTester/Source/MessageBusTester/Public/IMessageBusTester.h
deleted file mode 100644
diff --git a/Engine/Plugins/Experimental/MessageBusTester/Source/MessageBusTester/Public/IMessageBusTesterLogger.h b/Engine/Plugins/Experimental/MessageBusTester/Source/MessageBusTester/Public/IMessageBusTesterLogger.h
deleted file mode 100644
diff --git a/Engine/Plugins/Experimental/MessageBusTester/Source/MessageBusTester/Public/IMessageBusTesterModule.h b/Engine/Plugins/Experimental/MessageBusTester/Source/MessageBusTester/Public/IMessageBusTesterModule.h
deleted file mode 100644
diff --git a/Engine/Plugins/Experimental/MessageBusTester/Source/MessageBusTester/Public/MessageBusTesterCommon.h b/Engine/Plugins/Experimental/MessageBusTester/Source/MessageBusTester/Public/MessageBusTesterCommon.h
deleted file mode 100644
diff --git a/Engine/Plugins/Experimental/MessageBusTester/Source/MessageBusTester/Public/MessageBusTesterSettings.h b/Engine/Plugins/Experimental/MessageBusTester/Source/MessageBusTester/Public/MessageBusTesterSettings.h
deleted file mode 100644
diff --git a/Engine/Plugins/Experimental/MessageBusTester/Source/MessageBusTesterEditor/MessageBusTesterEditor.Build.cs b/Engine/Plugins/Experimental/MessageBusTester/Source/MessageBusTesterEditor/MessageBusTesterEditor.Build.cs
deleted file mode 100644
diff --git a/Engine/Plugins/Experimental/MessageBusTester/Source/MessageBusTesterEditor/Private/MessageBusTesterEditorModule.cpp b/Engine/Plugins/Experimental/MessageBusTester/Source/MessageBusTesterEditor/Private/MessageBusTesterEditorModule.cpp
deleted file mode 100644
diff --git a/Engine/Plugins/Experimental/MessageBusTester/Source/MessageBusTesterEditor/Private/MessageBusTesterEditorModule.h b/Engine/Plugins/Experimental/MessageBusTester/Source/MessageBusTesterEditor/Private/MessageBusTesterEditorModule.h
deleted file mode 100644
diff --git a/Engine/Plugins/Experimental/MessageBusTester/Source/MessageBusTesterEditor/Private/MessageBusTesterEditorStyle.cpp b/Engine/Plugins/Experimental/MessageBusTester/Source/MessageBusTesterEditor/Private/MessageBusTesterEditorStyle.cpp
deleted file mode 100644
diff --git a/Engine/Plugins/Experimental/MessageBusTester/Source/MessageBusTesterEditor/Private/MessageBusTesterEditorStyle.h b/Engine/Plugins/Experimental/MessageBusTester/Source/MessageBusTesterEditor/Private/MessageBusTesterEditorStyle.h
deleted file mode 100644
diff --git a/Engine/Plugins/Experimental/MessageBusTester/Source/MessageBusTesterEditor/Private/SMessageBusTestLogger.cpp b/Engine/Plugins/Experimental/MessageBusTester/Source/MessageBusTesterEditor/Private/SMessageBusTestLogger.cpp
deleted file mode 100644
diff --git a/Engine/Plugins/Experimental/MessageBusTester/Source/MessageBusTesterEditor/Private/SMessageBusTestLogger.h b/Engine/Plugins/Experimental/MessageBusTester/Source/MessageBusTesterEditor/Private/SMessageBusTestLogger.h
deleted file mode 100644
diff --git a/Engine/Plugins/Experimental/MessageBusTester/Source/MessageBusTesterEditor/Private/SMessageBusTestNetwork.cpp b/Engine/Plugins/Experimental/MessageBusTester/Source/MessageBusTesterEditor/Private/SMessageBusTestNetwork.cpp
deleted file mode 100644
diff --git a/Engine/Plugins/Experimental/MessageBusTester/Source/MessageBusTesterEditor/Private/SMessageBusTestNetwork.h b/Engine/Plugins/Experimental/MessageBusTester/Source/MessageBusTesterEditor/Private/SMessageBusTestNetwork.h
deleted file mode 100644
diff --git a/Engine/Plugins/Experimental/MessageBusTester/Source/MessageBusTesterEditor/Private/SMessageBusTestPlan.cpp b/Engine/Plugins/Experimental/MessageBusTester/Source/MessageBusTesterEditor/Private/SMessageBusTestPlan.cpp
deleted file mode 100644
diff --git a/Engine/Plugins/Experimental/MessageBusTester/Source/MessageBusTesterEditor/Private/SMessageBusTestPlan.h b/Engine/Plugins/Experimental/MessageBusTester/Source/MessageBusTesterEditor/Private/SMessageBusTestPlan.h
deleted file mode 100644
diff --git a/Engine/Plugins/Experimental/MessageBusTester/Source/MessageBusTesterEditor/Private/SMessageBusTesterPanel.cpp b/Engine/Plugins/Experimental/MessageBusTester/Source/MessageBusTesterEditor/Private/SMessageBusTesterPanel.cpp
deleted file mode 100644
diff --git a/Engine/Plugins/Experimental/MessageBusTester/Source/MessageBusTesterEditor/Private/SMessageBusTesterPanel.h b/Engine/Plugins/Experimental/MessageBusTester/Source/MessageBusTesterEditor/Private/SMessageBusTesterPanel.h
deleted file mode 100644
diff --git a/Engine/Plugins/Experimental/MessageBusTester/Source/MessageBusTesterEditor/Private/Widgets/SAddTestPlanEntryDialog.cpp b/Engine/Plugins/Experimental/MessageBusTester/Source/MessageBusTesterEditor/Private/Widgets/SAddTestPlanEntryDialog.cpp
deleted file mode 100644
diff --git a/Engine/Plugins/Experimental/MessageBusTester/Source/MessageBusTesterEditor/Private/Widgets/SAddTestPlanEntryDialog.h b/Engine/Plugins/Experimental/MessageBusTester/Source/MessageBusTesterEditor/Private/Widgets/SAddTestPlanEntryDialog.h
deleted file mode 100644
diff --git a/Engine/Plugins/Experimental/MessageBusTester/Source/MessageBusTesterEditor/Private/Widgets/SDiscoveredTesterListView.cpp b/Engine/Plugins/Experimental/MessageBusTester/Source/MessageBusTesterEditor/Private/Widgets/SDiscoveredTesterListView.cpp
deleted file mode 100644
diff --git a/Engine/Plugins/Experimental/MessageBusTester/Source/MessageBusTesterEditor/Private/Widgets/SDiscoveredTesterListView.h b/Engine/Plugins/Experimental/MessageBusTester/Source/MessageBusTesterEditor/Private/Widgets/SDiscoveredTesterListView.h
deleted file mode 100644
diff --git a/Engine/Plugins/Experimental/MessageBusTester/Source/MessageBusTesterEditor/Private/Widgets/SMessageBusTestPlanListView.cpp b/Engine/Plugins/Experimental/MessageBusTester/Source/MessageBusTesterEditor/Private/Widgets/SMessageBusTestPlanListView.cpp
deleted file mode 100644
diff --git a/Engine/Plugins/Experimental/MessageBusTester/Source/MessageBusTesterEditor/Private/Widgets/SMessageBusTestPlanListView.h b/Engine/Plugins/Experimental/MessageBusTester/Source/MessageBusTesterEditor/Private/Widgets/SMessageBusTestPlanListView.h
deleted file mode 100644
diff --git a/Engine/Plugins/Experimental/MessageBusTester/Source/MessageBusTesterEditor/Public/IMessageBusTesterEditorModule.h b/Engine/Plugins/Experimental/MessageBusTester/Source/MessageBusTesterEditor/Public/IMessageBusTesterEditorModule.h
deleted file mode 100644
diff --git a/Engine/Source/Programs/LiveLinkHub/Source/LiveLinkHub.Target.cs b/Engine/Source/Programs/LiveLinkHub/Source/LiveLinkHub.Target.cs
--- a/Engine/Source/Programs/LiveLinkHub/Source/LiveLinkHub.Target.cs
+++ b/Engine/Source/Programs/LiveLinkHub/Source/LiveLinkHub.Target.cs
@@ -48,7 +48,6 @@ public class LiveLinkHubTarget : TargetRules
 			"PropertyAccessEditor",
 			"PythonScriptPlugin",
 			"UdpMessaging",
-			"MessageBusTester",
 			"CameraCalibrationCore",
 			"AppleARKitFaceSupport",
 			"XInputDevice",
