LiveLinkHub - Uniformize white text color across the app

#rnx

diff --git a/Engine/Plugins/Animation/LiveLinkHub/Source/LiveLinkHub/Private/UI/Widgets/SLiveLinkHubSettings.h b/Engine/Plugins/Animation/LiveLinkHub/Source/LiveLinkHub/Private/UI/Widgets/SLiveLinkHubSettings.h
--- a/Engine/Plugins/Animation/LiveLinkHub/Source/LiveLinkHub/Private/UI/Widgets/SLiveLinkHubSettings.h
+++ b/Engine/Plugins/Animation/LiveLinkHub/Source/LiveLinkHub/Private/UI/Widgets/SLiveLinkHubSettings.h
@@ -68,6 +68,7 @@ public:
 					.ButtonContent()
 					[
 						SNew(SImage)
+						.ColorAndOpacity(FSlateColor::UseForeground())
 						.Image(FAppStyle::Get().GetBrush("FullBlueprintEditor.EditGlobalOptions"))
 					]
 					.OnGetMenuContent(this, &SLiveLinkHubSettings::OnSettingsComboBoxOpen)
diff --git a/Engine/Plugins/Animation/LiveLinkHub/Source/LiveLinkHub/Private/UI/Widgets/SLiveLinkTimecode.cpp b/Engine/Plugins/Animation/LiveLinkHub/Source/LiveLinkHub/Private/UI/Widgets/SLiveLinkTimecode.cpp
--- a/Engine/Plugins/Animation/LiveLinkHub/Source/LiveLinkHub/Private/UI/Widgets/SLiveLinkTimecode.cpp
+++ b/Engine/Plugins/Animation/LiveLinkHub/Source/LiveLinkHub/Private/UI/Widgets/SLiveLinkTimecode.cpp
@@ -138,6 +138,7 @@ void SLiveLinkTimecode::Construct(const FArguments& InArgs)
 			.ToolTipText(LOCTEXT("TimeAndSyncToolTip", "Open Time and Sync settings"))
 			[
 				SNew(SImage)
+				.ColorAndOpacity(FSlateColor::UseForeground())
 				.Image(StyleSet->GetBrush("LiveLinkHub.EllipsisIcon"))
 			]
 		]
