LiveLinkHub - Convert LiveLinkHubWindowController to an IEditorMainFrameHandler so that the MainFrame module registers a root window.

This change will allow LiveLinkHub to instantiate asset editors, such as for Media Profile.

todo: Convert existing mode menus to UToolMenu entries, add toolbar buttons for MediaProfile and the mode switcher

diff --git a/Engine/Plugins/Animation/LiveLinkHub/Source/LiveLinkHub/Private/LiveLinkHubApplication.h b/Engine/Plugins/Animation/LiveLinkHub/Source/LiveLinkHub/Private/LiveLinkHubApplication.h
--- a/Engine/Plugins/Animation/LiveLinkHub/Source/LiveLinkHub/Private/LiveLinkHubApplication.h
+++ b/Engine/Plugins/Animation/LiveLinkHub/Source/LiveLinkHub/Private/LiveLinkHubApplication.h
@@ -35,7 +35,7 @@ public:
 
 	/** Get the name of the TabLayout for a given mode. */
 	FName GetLayoutName(FName ModeName);
-	
+
 	/** Get the list of modes that were registered with the app. */
 	TArray<FName> GetApplicationModes() const;
 
@@ -92,6 +92,7 @@ protected:
 	virtual void MapToolkitCommands() {}
 	virtual FText GetToolkitName() const override;
 	virtual FText GetToolkitToolTipText() const override;
+	virtual FName GetOverrideTabId() const override { return "LiveLinkHub"; }
 	//~ End AssetEditorToolkit interface
 
 	/** Iterate through layout directories to find .llhlayout files. */
diff --git a/Engine/Plugins/Animation/LiveLinkHub/Source/LiveLinkHub/Private/UI/Window/LiveLinkHubWindowController.cpp b/Engine/Plugins/Animation/LiveLinkHub/Source/LiveLinkHub/Private/UI/Window/LiveLinkHubWindowController.cpp
--- a/Engine/Plugins/Animation/LiveLinkHub/Source/LiveLinkHub/Private/UI/Window/LiveLinkHubWindowController.cpp
+++ b/Engine/Plugins/Animation/LiveLinkHub/Source/LiveLinkHub/Private/UI/Window/LiveLinkHubWindowController.cpp
@@ -4,10 +4,12 @@
 
 #include "Clients/LiveLinkHubProvider.h"
 #include "CoreGlobals.h"
+#include "Framework/Application/SWindowTitleBar.h"
 #include "Features/IModularFeatures.h"
 #include "Framework/Application/SlateApplication.h"
 #include "Framework/Docking/LayoutService.h"
 #include "Framework/Notifications/NotificationManager.h"
+#include "Interfaces/IMainFrameModule.h"
 #include "LiveLinkHub.h"
 #include "HAL/PlatformApplicationMisc.h"
 #include "Misc/ConfigCacheIni.h"
@@ -19,9 +21,26 @@
 
 #define LOCTEXT_NAMESPACE "LiveLinkHubWindowController"
 
+const FName FLiveLinkHubWindowController::LiveLinkHubTabId = "LiveLinkHub";
+
 FLiveLinkHubWindowController::FLiveLinkHubWindowController()
 {
-	ModalWindowManager = InitializeSlateApplication();
+	IModularFeatures::Get().RegisterModularFeature(IEditorMainFrameProvider::GetModularFeatureName(), this);
+
+	const FName LayoutName = TEXT("LiveLinkHub_v1.2");
+	DefaultLayout = FTabManager::NewLayout(LayoutName)
+		->AddArea
+		(
+			// Toolkits window
+			FTabManager::NewPrimaryArea()
+			->Split
+			(
+				FTabManager::NewStack()
+				->SetSizeCoefficient(1.0f)
+				->AddTab(LiveLinkHubTabId, ETabState::ClosedTab)
+				->AddTab("StandaloneToolkit", ETabState::ClosedTab) // This tab is placed here in case the user opens an asset editor (ie. Media Profile)
+			)
+		);
 }
 
 FLiveLinkHubWindowController::~FLiveLinkHubWindowController()
@@ -30,79 +49,43 @@ FLiveLinkHubWindowController::~FLiveLinkHubWindowController()
 	{
 		RootWindow->SetOnWindowClosed(nullptr);
 	}
+
+	IModularFeatures::Get().UnregisterModularFeature(IEditorMainFrameProvider::GetModularFeatureName(), this);
 }
 
-TSharedRef<SWindow> FLiveLinkHubWindowController::CreateWindow()
+TSharedRef<SWidget> FLiveLinkHubWindowController::CreateMainFrameContentWidget() const
 {
-	FDisplayMetrics DisplayMetrics;
-	FSlateApplication::Get().GetDisplayMetrics(DisplayMetrics);
-	const float DPIScaleFactor = FPlatformApplicationMisc::GetDPIScaleFactorAtPoint(DisplayMetrics.PrimaryDisplayWorkAreaRect.Left, DisplayMetrics.PrimaryDisplayWorkAreaRect.Top);
-
-	constexpr bool bEmbedTitleAreaContent = true;
-	const FVector2D ClientSize(1200.0f * DPIScaleFactor, 800.0f * DPIScaleFactor);
-	TSharedRef<SWindow> RootWindowRef = SNew(SWindow)
-		.Title(LOCTEXT("WindowTitle", "Live Link Hub"))
-		.CreateTitleBar(!bEmbedTitleAreaContent)
-		.SupportsMaximize(true)
-		.SupportsMinimize(true)
-		.IsInitiallyMaximized(false)
-		.IsInitiallyMinimized(false)
-		.SizingRule(ESizingRule::UserSized)
-		.AutoCenter(EAutoCenter::PreferredWorkArea)
-		.ClientSize(ClientSize)
-		.AdjustInitialSizeAndPositionForDPIScale(false);
-
-	RootWindow = RootWindowRef;
-		
-	constexpr bool bShowRootWindowImmediately = false;
-	FSlateApplication::Get().AddWindow(RootWindowRef, bShowRootWindowImmediately);
-	FGlobalTabmanager::Get()->SetRootWindow(RootWindowRef);
-	FGlobalTabmanager::Get()->SetAllowWindowMenuBar(true);
-	FSlateNotificationManager::Get().SetRootWindow(RootWindowRef);	
-
-
-	RootWindow->SetRequestDestroyWindowOverride(FRequestDestroyWindowOverride::CreateRaw(this, &FLiveLinkHubWindowController::CloseRootWindowOverride));
-	RootWindow->SetOnWindowClosed(FOnWindowClosed::CreateRaw(this, &FLiveLinkHubWindowController::OnWindowClosed));
+	// We defer creation until RestoreLayout is called so that the widget has a valid RootWindow to use for restoring layout.
+	return SNullWidget::NullWidget;
+}
 
-	return RootWindowRef;
+FMainFrameWindowOverrides FLiveLinkHubWindowController::GetDesiredWindowConfiguration() const
+{
+	FMainFrameWindowOverrides Overrides;
+	Overrides.WindowSize = FVector2D(1200.0f, 800.0f);
+	Overrides.WindowTitle = LOCTEXT("WindowTitle", "Live Link Hub");
+	Overrides.bEmbedTitleAreaContent = true;
+	Overrides.bInitiallyMaximized = false;
+	return Overrides;
 }
+
 void FLiveLinkHubWindowController::RestoreLayout(TSharedPtr<FAssetEditorToolkit> AssetEditorToolkit)
 {
 	TRACE_CPUPROFILER_EVENT_SCOPE(FLiveLinkHubWindowController::RestoreLayout);
 
-	static const FTabId StandaloneTabId("StandaloneToolkit");
+	FModuleManager::Get().LoadModuleChecked<IMainFrameModule>("MainFrame").CreateDefaultMainFrame(false, false);
 
-	const FName LayoutName = TEXT("LiveLinkHub_v1.1");
-	const TSharedRef<FTabManager::FLayout> DefaultLayout =
-		FTabManager::NewLayout(LayoutName)
-		->AddArea
-		(
-			// Toolkits window
-			FTabManager::NewPrimaryArea()
-			->Split
-			(
-				FTabManager::NewStack()
-				->SetSizeCoefficient(1.0f)
-				->AddTab(StandaloneTabId, ETabState::ClosedTab)
-			)
-		);
-
-	{
-		TRACE_CPUPROFILER_EVENT_SCOPE(FLiveLinkHubWindowController::RestoreFrom);
-		constexpr bool bEmbedTitleAreaContent = true;
-		const TSharedPtr<SWidget> Content = FGlobalTabmanager::Get()->RestoreFrom(DefaultLayout, RootWindow, bEmbedTitleAreaContent, EOutputCanBeNullptr::Never);
-		RootWindow->SetContent(Content.ToSharedRef());
-	}
+	RootWindow = FGlobalTabmanager::Get()->GetRootWindow();
+	RootWindow->SetRequestDestroyWindowOverride(FRequestDestroyWindowOverride::CreateRaw(this, &FLiveLinkHubWindowController::CloseRootWindowOverride));
+	RootWindow->SetOnWindowClosed(FOnWindowClosed::CreateRaw(this, &FLiveLinkHubWindowController::OnWindowClosed));
 
-	RootWindow->ShowWindow();
-	constexpr bool bForceWindowToFront = true;
-	RootWindow->BringToFront(bForceWindowToFront);
+	constexpr bool bEmbedTitleAreaContent = true;
+	const TSharedPtr<SWidget> Content = FGlobalTabmanager::Get()->RestoreFrom(DefaultLayout.ToSharedRef(), RootWindow, bEmbedTitleAreaContent, EOutputCanBeNullptr::Never);
+	RootWindow->SetContent(Content.ToSharedRef());
 
 	// Pass a dummy object to the asset editor since we're not actually editing an object.
 	UObject* DummyObject = GetMutableDefault<UObject>();
-	AssetEditorToolkit->InitAssetEditor(EToolkitMode::Standalone, nullptr, "LiveLinkHub", DefaultLayout, /*bCreateDefaultStandaloneMenu*/true, /*bCreateDefaultToolbar*/true, DummyObject, /*bInIsToolbarFocusable*/true, true, {});
-
-	FGlobalTabmanager::Get()->SetMainTab(StandaloneTabId);
+	AssetEditorToolkit->InitAssetEditor(EToolkitMode::Standalone, nullptr, LiveLinkHubTabId, DefaultLayout.ToSharedRef(), /*bCreateDefaultStandaloneMenu*/true, /*bCreateDefaultToolbar*/true, DummyObject, /*bInIsToolbarFocusable*/true, true, {});
 
 	TSharedRef<SWidget> ProviderNameWidget =
 		SNew(SBox)
@@ -126,21 +109,12 @@ void FLiveLinkHubWindowController::RestoreLayout(TSharedPtr<FAssetEditorToolkit>
 			SNew(SLiveLinkHubMemoryStats)
 		];
 
-	TSharedPtr<SDockTab> MainTab = FGlobalTabmanager::Get()->FindExistingLiveTab(StandaloneTabId);
+	FGlobalTabmanager::Get()->SetMainTab(LiveLinkHubTabId);
+	TSharedPtr<SDockTab> MainTab = FGlobalTabmanager::Get()->FindExistingLiveTab(LiveLinkHubTabId);
 	MainTab->SetTitleBarRightContent(RightMenuWidget);
 	MainTab->SetTabIcon(FSlateIcon("LiveLinkStyle", "LiveLinkHub.Icon.Small").GetIcon());
 }
 
-TSharedPtr<FModalWindowManager> FLiveLinkHubWindowController::InitializeSlateApplication()
-{
-	TRACE_CPUPROFILER_EVENT_SCOPE(FLiveLinkHubWindowController::InitializeSlateApplication);
-
-	const FText ApplicationTitle = LOCTEXT("AppTitle", "Live Link Hub");
-	FGlobalTabmanager::Get()->SetApplicationTitle(ApplicationTitle);
-
-	return MakeShared<FModalWindowManager>(CreateWindow());
-}
-
 void FLiveLinkHubWindowController::CloseRootWindowOverride(const TSharedRef<SWindow>& Window)
 {
 	if (GetDefault<ULiveLinkHubSettings>()->bConfirmClose)
diff --git a/Engine/Plugins/Animation/LiveLinkHub/Source/LiveLinkHub/Private/UI/Window/LiveLinkHubWindowController.h b/Engine/Plugins/Animation/LiveLinkHub/Source/LiveLinkHub/Private/UI/Window/LiveLinkHubWindowController.h
--- a/Engine/Plugins/Animation/LiveLinkHub/Source/LiveLinkHub/Private/UI/Window/LiveLinkHubWindowController.h
+++ b/Engine/Plugins/Animation/LiveLinkHub/Source/LiveLinkHub/Private/UI/Window/LiveLinkHubWindowController.h
@@ -2,20 +2,17 @@
 
 #pragma once
 
+#include "Interfaces/IEditorMainFrameProvider.h"
+
 #include "Framework/Docking/TabManager.h"
 
 class FAssetEditorToolkit;
-class FLiveLinkHubSessionTabBase;
-class FOutputLogController;
-class ILiveLinkClient;
-class FModalWindowManager;
 class SWindow;
 
 /** Responsible for creating the Slate window for the hub. */
-class FLiveLinkHubWindowController : public TSharedFromThis<FLiveLinkHubWindowController>
+class FLiveLinkHubWindowController : public TSharedFromThis<FLiveLinkHubWindowController>, public IEditorMainFrameProvider
 {
 public:
-	
 	FLiveLinkHubWindowController();
 	~FLiveLinkHubWindowController();
 
@@ -23,27 +20,25 @@ public:
 	TSharedPtr<SWindow> GetRootWindow() const { return RootWindow; }
 	/** Restore the window's layout from a config. */
 	void RestoreLayout(TSharedPtr<FAssetEditorToolkit> AssetEditorToolkit);
+
+	//~ Begin IEditorMainFrameProvider
+	virtual bool IsRequestingMainFrameControl() const { return true; }
+	virtual FMainFrameWindowOverrides GetDesiredWindowConfiguration() const;
+	virtual TSharedRef<SWidget> CreateMainFrameContentWidget() const;
+	//~ End IEditorMainFrameProvider
+
 private:
 	/** Create the main window. */
 	TSharedRef<SWindow> CreateWindow();
-	/** Create the slate application that hosts the livelink hub. */
-	TSharedPtr<FModalWindowManager> InitializeSlateApplication();
 	/** Override to handle confirming if the user wants to quit. */
 	void CloseRootWindowOverride(const TSharedRef<SWindow>& Window);
 	/** Window closed handler. */
 	void OnWindowClosed(const TSharedRef<SWindow>& Window);
 private:
-	/** Pointer to the livelink client. */
-	ILiveLinkClient* LiveLinkClient = nullptr;
-	/** The ini file to use for saving the layout */
-	FString LiveLinkHubLayoutIni;
-	/** Holds the current layout for saving later. */
-	TSharedPtr<FTabManager::FLayout> PersistentLayout;
 	/** The main window being managed */
 	TSharedPtr<SWindow> RootWindow;
-	/** Manages modal windows for the application. */
-	TSharedPtr<FModalWindowManager> ModalWindowManager;
-	/** Menu bar widget for the hub. */
-	TSharedPtr<class SWindowTitleBar> WindowTitleBar;
-
+	/** Default layout for the app. */
+	TSharedPtr<FTabManager::FLayout> DefaultLayout;
+	/** Tab ID of the main LiveLinkHub  */
+	static const FName LiveLinkHubTabId;
 };
