LiveLinkHub - Add MediaProfile support, enable Audio Capture timecode provider

#rnx

diff --git a/Engine/Plugins/Animation/LiveLinkHub/Source/LiveLinkHub/Private/LiveLinkHubApplicationMode.cpp b/Engine/Plugins/Animation/LiveLinkHub/Source/LiveLinkHub/Private/LiveLinkHubApplicationMode.cpp
--- a/Engine/Plugins/Animation/LiveLinkHub/Source/LiveLinkHub/Private/LiveLinkHubApplicationMode.cpp
+++ b/Engine/Plugins/Animation/LiveLinkHub/Source/LiveLinkHub/Private/LiveLinkHubApplicationMode.cpp
@@ -43,26 +43,8 @@ void FLiveLinkHubApplicationMode::PostActivateMode()
 {
 	FToolMenuOwnerScoped OwnerScoped(GetModeName());
 
-
 	UToolMenu* Menu = UToolMenus::Get()->ExtendMenu(TEXT("LiveLinkHub.ToolBar"));
 
-	FToolMenuSection& GlobalWidgetsSection = Menu->AddSection(TEXT("GlobalWidgets"), FText::GetEmpty(), FToolMenuInsert(NAME_None, EToolMenuInsertType::First));
-
-	TSharedRef<SWidget> GlobalWidget =
-		SNew(SHorizontalBox)
-		+ SHorizontalBox::Slot()
-		.AutoWidth()
-		.Padding(2.0, 0.0)
-		.VAlign(VAlign_Center)
-		[
-			SNew(SLiveLinkHubTopologyModeSwitcher)
-		];
-
-	GlobalWidgetsSection.AddEntry(
-		FToolMenuEntry::InitWidget(TEXT("TestWidget"), GlobalWidget, FText::GetEmpty(), true, false)
-	);
-	GlobalWidgetsSection.AddSeparator(NAME_None);
-
 	FToolMenuSection& SessionConfigSection = Menu->AddSection(TEXT("ModeWidgets"), FText::GetEmpty(), FToolMenuInsert(TEXT("GlobalWidgets"), EToolMenuInsertType::After));
 
 	TSharedRef<SHorizontalBox> ModeWidgets = SNew(SHorizontalBox);
diff --git a/Engine/Source/Programs/LiveLinkHub/Config/DefaultEngine.ini b/Engine/Source/Programs/LiveLinkHub/Config/DefaultEngine.ini
--- a/Engine/Source/Programs/LiveLinkHub/Config/DefaultEngine.ini
+++ b/Engine/Source/Programs/LiveLinkHub/Config/DefaultEngine.ini
@@ -72,3 +72,6 @@ NodeMemoryLoggingThresholdMB=5
 +PropertyRedirects=(OldName="LiveLinkHubSettings.bAutoAddDiscoveredClients",NewName="AutoConnectClients")
 +ClassRedirects=(OldName="/Script/LiveLinkHub.LiveLinkHubSubjectSettings",NewName="/Script/LiveLink.LiveLinkHubSubjectSettings")
 
+[MediaProfile]
+ButtonMenu="AssetEditor.LiveLinkHub.ToolBar"
+ButtonAlignment="Last"
\ No newline at end of file
diff --git a/Engine/Source/Programs/LiveLinkHub/Config/DefaultGame.ini b/Engine/Source/Programs/LiveLinkHub/Config/DefaultGame.ini
--- a/Engine/Source/Programs/LiveLinkHub/Config/DefaultGame.ini
+++ b/Engine/Source/Programs/LiveLinkHub/Config/DefaultGame.ini
@@ -20,3 +20,8 @@ bTranslateRebroadcastFrames=true
 ProjectID=914C155297C9471A9E34100517697A75
 ProjectName=LiveLinkHub
 
+[/Script/MediaFrameworkUtilities.MediaProfileSettings]
+bApplyInCommandlet=False
+bAutosaveAfterChange=True
+StartupMediaProfile=None
+
diff --git a/Engine/Source/Programs/LiveLinkHub/Source/LiveLinkHub.Target.cs b/Engine/Source/Programs/LiveLinkHub/Source/LiveLinkHub.Target.cs
--- a/Engine/Source/Programs/LiveLinkHub/Source/LiveLinkHub.Target.cs
+++ b/Engine/Source/Programs/LiveLinkHub/Source/LiveLinkHub.Target.cs
@@ -35,10 +35,10 @@ public class LiveLinkHubTarget : TargetRules
 		EnablePlugins.AddRange(new string[]
 		{
 			"ACLPlugin",
+			"AudioCaptureTimecodeProvider",
 			"LiveLink",
 			"LiveLinkHub",
 			"PluginBrowser",
-
 			"LiveLinkCamera",
 			"LiveLinkLens",
 			"LensComponent",
@@ -63,6 +63,7 @@ public class LiveLinkHubTarget : TargetRules
 			"CaptureManagerApp",
 			"CaptureManagerDevices",
 			"MediaIOFramework",
+			"MediaFrameworkUtilities",
 			"AjaMedia",
 			"BlackmagicMedia"
 		});
