LiveLinkHub - Convert the creator app mode to use the UToolMenu directly to allow Media Profile menu entry to insert itself in the toolbar

#rnx

diff --git a/Engine/Plugins/Animation/LiveLinkHub/Source/LiveLinkHub/Private/LiveLinkHubCreatorAppMode.h b/Engine/Plugins/Animation/LiveLinkHub/Source/LiveLinkHub/Private/LiveLinkHubCreatorAppMode.h
--- a/Engine/Plugins/Animation/LiveLinkHub/Source/LiveLinkHub/Private/LiveLinkHubCreatorAppMode.h
+++ b/Engine/Plugins/Animation/LiveLinkHub/Source/LiveLinkHub/Private/LiveLinkHubCreatorAppMode.h
@@ -406,15 +406,33 @@ public:
 		return FSlateIcon("LiveLinkStyle", TEXT("LiveLinkHub.Subjects.Icon"));
 	}
 
-
-	virtual TArray<TSharedRef<SWidget>> GetToolbarWidgets_Impl() override
+	virtual void PostActivateMode() override
 	{
+		FLiveLinkHubApplicationMode::PostActivateMode();
+
+		FToolMenuOwnerScoped OwnerScoped(GetModeName());
+
+		UToolMenu* Toolbar = UToolMenus::Get()->ExtendMenu(TEXT("AssetEditor.LiveLinkHub.ToolBar"));
+		
+		FToolMenuSection& MiddleSection = Toolbar->AddSection("SlateWidgets");
+
+		MiddleSection.AddEntry(
+			FToolMenuEntry::InitWidget("SlateWidget", SNew(SLiveLinkRecordingSessionInfo), FText::GetEmpty(), true, false)
+		);
+		MiddleSection.Alignment = EToolMenuSectionAlign::Last;
+
+
+		FToolMenuSection& RightSection = Toolbar->AddSection("TimecodeWidgets", FText::GetEmpty(), FToolMenuInsert{ "SlateWidgets", EToolMenuInsertType::After });
+		RightSection.Alignment = EToolMenuSectionAlign::Last;
+
 		FLiveLinkHubModule& LiveLinkHubModule = FModuleManager::Get().GetModuleChecked<FLiveLinkHubModule>("LiveLinkHub");
-		return {
-			SNew(SBox).Padding(FMargin(0,0,20,0)) [ SNew(SLiveLinkRecordingSessionInfo) ],
-			LiveLinkHubModule.GetRecordingController()->MakeRecordToolbarEntry(),
-			SNew(SLiveLinkTimecode)
-		}; 
+		RightSection.AddEntry(
+			FToolMenuEntry::InitWidget(TEXT("RecordButton"), LiveLinkHubModule.GetRecordingController()->MakeRecordToolbarEntry(), FText::GetEmpty(), true, false)
+		);
+
+		RightSection.AddEntry(
+			FToolMenuEntry::InitWidget(TEXT("Timecode"), SNew(SLiveLinkTimecode), FText::GetEmpty(), true, false)
+		);
 	}
 
 	virtual TArray<TSharedRef<SWidget>> GetStatusBarWidgets_Impl() override
