LiveLinkHub - Move MediaProfile button next to the settings button

#rnx

diff --git a/Engine/Plugins/Animation/LiveLinkHub/Source/LiveLinkHub/Private/LiveLinkHubCreatorAppMode.h b/Engine/Plugins/Animation/LiveLinkHub/Source/LiveLinkHub/Private/LiveLinkHubCreatorAppMode.h
--- a/Engine/Plugins/Animation/LiveLinkHub/Source/LiveLinkHub/Private/LiveLinkHubCreatorAppMode.h
+++ b/Engine/Plugins/Animation/LiveLinkHub/Source/LiveLinkHub/Private/LiveLinkHubCreatorAppMode.h
@@ -427,12 +427,14 @@ public:
 
 		FLiveLinkHubModule& LiveLinkHubModule = FModuleManager::Get().GetModuleChecked<FLiveLinkHubModule>("LiveLinkHub");
 		RightSection.AddEntry(
-			FToolMenuEntry::InitWidget(TEXT("RecordButton"), LiveLinkHubModule.GetRecordingController()->MakeRecordToolbarEntry(), FText::GetEmpty(), true, false)
+			FToolMenuEntry::InitWidget("RecordButton", LiveLinkHubModule.GetRecordingController()->MakeRecordToolbarEntry(), FText::GetEmpty(), true, false)
 		);
 
 		RightSection.AddEntry(
-			FToolMenuEntry::InitWidget(TEXT("Timecode"), SNew(SLiveLinkTimecode), FText::GetEmpty(), true, false)
+			FToolMenuEntry::InitWidget("Timecode", SNew(SLiveLinkTimecode), FText::GetEmpty(), true, false)
 		);
+
+		Toolbar->AddSection("MediaProfile", FText::GetEmpty(), FToolMenuInsert(NAME_None, EToolMenuInsertType::Last));
 	}
 
 	virtual TArray<TSharedRef<SWidget>> GetStatusBarWidgets_Impl() override
