LiveLinkHub - Allow deleting and upgrading recordings in place

#Synced-CL 47110913

diff --git a/Engine/Plugins/Animation/LiveLinkHub/Source/LiveLinkHub/Private/LiveLinkHub.cpp b/Engine/Plugins/Animation/LiveLinkHub/Source/LiveLinkHub/Private/LiveLinkHub.cpp
--- a/Engine/Plugins/Animation/LiveLinkHub/Source/LiveLinkHub/Private/LiveLinkHub.cpp
+++ b/Engine/Plugins/Animation/LiveLinkHub/Source/LiveLinkHub/Private/LiveLinkHub.cpp
@@ -395,9 +395,9 @@ bool FLiveLinkHub::CanSetTopologyMode() const
 	return !OverrideTopologyMode.IsSet();
 }
 
-void FLiveLinkHub::UpgradeAndSaveRecording(ULiveLinkUAssetRecording* InRecording)
+void FLiveLinkHub::UpgradeAndSaveRecording(ULiveLinkUAssetRecording* InRecording, bool bUpgradeInPlace)
 {
-	GetPlaybackController()->OnPlaybackReady().AddSP(this, &FLiveLinkHub::OnPlaybackReadyForUpgrade);
+	GetPlaybackController()->OnPlaybackReady().AddSP(this, &FLiveLinkHub::OnPlaybackReadyForUpgrade, bUpgradeInPlace);
 	GetPlaybackController()->PreparePlayback(InRecording);
 }
 
@@ -557,7 +557,7 @@ void FLiveLinkHub::OnSubjectOutboundNameModified(const FLiveLinkSubjectKey& Subj
 	LiveLinkProvider->RemoveSubject(*PreviousOutboundName);
 }
 
-void FLiveLinkHub::OnPlaybackReadyForUpgrade()
+void FLiveLinkHub::OnPlaybackReadyForUpgrade(bool bUpgradeInPlace)
 {
 	if (!PlaybackController || !RecordingController || !PlaybackController->GetRecording())
 	{
@@ -566,13 +566,16 @@ void FLiveLinkHub::OnPlaybackReadyForUpgrade()
 
 	PlaybackController->OnPlaybackReady().RemoveAll(this);
 
-	// Get the save path before starting the long task.
-	const FString DefaultName = PlaybackController->GetRecording()->GetName() + TEXT("_copy");
-	FString RecordingPath;
-	if (!RecordingController->OpenSaveAsPrompt(RecordingPath, DefaultName))
+	FString RecordingPath = PlaybackController->GetRecording()->GetPackage()->GetPathName();
+	if (!bUpgradeInPlace)
 	{
-		PlaybackController->EjectAndUnload();
-		return;
+		// Get the save path before starting the long task.
+		const FString DefaultName = PlaybackController->GetRecording()->GetName() + TEXT("_copy");
+		if (!RecordingController->OpenSaveAsPrompt(RecordingPath, DefaultName))
+		{
+			PlaybackController->EjectAndUnload();
+			return;
+		}
 	}
 		
 	ULiveLinkUAssetRecording* CurrentRecording = CastChecked<ULiveLinkUAssetRecording>(PlaybackController->GetRecording().Get());
diff --git a/Engine/Plugins/Animation/LiveLinkHub/Source/LiveLinkHub/Private/LiveLinkHub.h b/Engine/Plugins/Animation/LiveLinkHub/Source/LiveLinkHub/Private/LiveLinkHub.h
--- a/Engine/Plugins/Animation/LiveLinkHub/Source/LiveLinkHub/Private/LiveLinkHub.h
+++ b/Engine/Plugins/Animation/LiveLinkHub/Source/LiveLinkHub/Private/LiveLinkHub.h
@@ -99,9 +99,11 @@ public:
 	/** Delegate called when the topology mode changed. */
 	FOnTopologyModeChanged& OnTopologyModeChanged() { return TopologyModeChangedDelegate; }
 
-	/** Save an existing recording as a new file. */
-	void UpgradeAndSaveRecording(ULiveLinkUAssetRecording* InRecording);
-	
+	/** 
+	 * Upgrade a recording and save to a new file or over the existing file.
+	 */
+	void UpgradeAndSaveRecording(ULiveLinkUAssetRecording* InRecording, bool bUpgradeInPlace = false);
+
 	/** Delegate triggered whenever the Live Link Client receives frame data. */
 	FSimpleDelegate& OnFrameReceived() { return OnFrameReceivedDelegate; }
 
@@ -161,7 +163,7 @@ private:
 	void OnSubjectOutboundNameModified(const FLiveLinkSubjectKey& SubjectKey, const FString& PreviousOutboundName, const FString& NewOutboundName) const;
 
 	/** Callback when playback is ready because we are performing an upgrade. */
-	void OnPlaybackReadyForUpgrade();
+	void OnPlaybackReadyForUpgrade(bool bUpgradeInPlace);
 	
 private:
 	/** Implements the logic to manage the clients tabs. */
diff --git a/Engine/Plugins/Animation/LiveLinkHub/Source/LiveLinkHub/Private/Recording/Implementations/LiveLinkUAssetRecorder.cpp b/Engine/Plugins/Animation/LiveLinkHub/Source/LiveLinkHub/Private/Recording/Implementations/LiveLinkUAssetRecorder.cpp
--- a/Engine/Plugins/Animation/LiveLinkHub/Source/LiveLinkHub/Private/Recording/Implementations/LiveLinkUAssetRecorder.cpp
+++ b/Engine/Plugins/Animation/LiveLinkHub/Source/LiveLinkHub/Private/Recording/Implementations/LiveLinkUAssetRecorder.cpp
@@ -148,6 +148,7 @@ void FLiveLinkUAssetRecorder::SaveFromExistingRecording(const ULiveLinkRecording
 	FLiveLinkUAssetRecordingData InRecordingData = CastChecked<ULiveLinkUAssetRecording>(InRecording)->RecordingData;
 	SaveRecording(MoveTemp(InRecordingData), InRecording->LengthInSeconds, /*bIncrementTake*/ false,
 		InRecording->RecordingPreset, &InRecording->FrameRate, InFilePath);
+
 	CurrentRecording.Reset();
 }
 
@@ -322,33 +323,45 @@ void FLiveLinkUAssetRecorder::SaveRecording(FLiveLinkUAssetRecordingData&& InRec
 
 	// Saving into a new package
 	const FString NewAssetName = FPackageName::GetLongPackageAssetName(PackageName);
-	UPackage* NewPackage = CreatePackage(*PackageName);
+	UPackage* Package = FindPackage(NULL, *PackageName);
+	ULiveLinkUAssetRecording* Recording = nullptr;
+
+	if (Package)
+	{
+		// If the package already exists, then that means that we shouldn't create a new asset, and just overwrite the existing recording.
+		Recording = Cast<ULiveLinkUAssetRecording>(StaticFindObject(ULiveLinkUAssetRecording::StaticClass(), FTopLevelAssetPath(*PackageName, *NewAssetName), EFindObjectFlags::None));
+		check(Recording);
+	}
+	else
+	{
+		Package = CreatePackage(*PackageName);
+		Recording = NewObject<ULiveLinkUAssetRecording>(Package, *NewAssetName, RF_Public | RF_Standalone);
+	}
 
-	if (TObjectPtr<ULiveLinkUAssetRecording> NewRecording = NewObject<ULiveLinkUAssetRecording>(NewPackage, *NewAssetName, RF_Public | RF_Standalone))
+	if (Recording)
 	{
-		NewRecording->InitializeNewRecordingData(MoveTemp(InRecordingData), InRecordingLength);
+		Recording->InitializeNewRecordingData(MoveTemp(InRecordingData), InRecordingLength);
 		// If preset or frame rate provided make sure to set them after initializing new data.
 		if (InPreset)
 		{
 			// Copy the provided preset to our preset. This allows full source names to be preserved if we're saving data from another recording,
 			// such as during an upgrade.
-			UEngine::CopyPropertiesForUnrelatedObjects(InPreset, NewRecording->RecordingPreset);	
+			UEngine::CopyPropertiesForUnrelatedObjects(InPreset, Recording->RecordingPreset);
 		}
 		if (InFrameRate)
 		{
-			NewRecording->FrameRate = *InFrameRate;
+			Recording->FrameRate = *InFrameRate;
 		}
 
-		NewRecording->MarkPackageDirty();
-
-		FAssetRegistryModule::AssetCreated(NewRecording);
+		Recording->MarkPackageDirty();
+		FAssetRegistryModule::AssetCreated(Recording);
 
-		const TStrongObjectPtr<ULiveLinkUAssetRecording> NewRecordingStrongPtr(NewRecording);
+		const TStrongObjectPtr<ULiveLinkUAssetRecording> NewRecordingStrongPtr(Recording);
 		
 		// Create a task to run on a separate thread for saving all frame data and writing the final UAsset to disk.
 		// We use a container rather than just one task on the chance a save operation is still running when another recording is being saved.
 		FAsyncTask<FLiveLinkSaveRecordingAsyncTask>& AsyncTask = *AsyncSaveTasks.Add(NewRecordingStrongPtr,
-			MakeUnique<FAsyncTask<FLiveLinkSaveRecordingAsyncTask>>(NewRecording, this));
+			MakeUnique<FAsyncTask<FLiveLinkSaveRecordingAsyncTask>>(Recording, this));
 		AsyncTask.StartBackgroundTask();
 	}
 
diff --git a/Engine/Plugins/Animation/LiveLinkHub/Source/LiveLinkHub/Private/Recording/SLiveLinkHubRecordingListView.h b/Engine/Plugins/Animation/LiveLinkHub/Source/LiveLinkHub/Private/Recording/SLiveLinkHubRecordingListView.h
--- a/Engine/Plugins/Animation/LiveLinkHub/Source/LiveLinkHub/Private/Recording/SLiveLinkHubRecordingListView.h
+++ b/Engine/Plugins/Animation/LiveLinkHub/Source/LiveLinkHub/Private/Recording/SLiveLinkHubRecordingListView.h
@@ -15,6 +15,7 @@
 #include "LiveLinkHubModule.h"
 #include "Misc/PackageName.h"
 #include "Misc/Paths.h"
+#include "ObjectTools.h"
 #include "Recording/LiveLinkHubPlaybackController.h"
 #include "Recording/LiveLinkRecording.h"
 #include "UObject/SavePackage.h"
@@ -231,6 +232,7 @@ private:
 			AssetPickerConfig.Filter = MakeAssetFilter();
 			AssetPickerConfig.OnAssetDoubleClicked = FOnAssetSelected::CreateRaw(this, &SLiveLinkHubRecordingListView::OnImportRecording);
 			AssetPickerConfig.OnGetAssetContextMenu = FOnGetAssetContextMenu::CreateRaw(this, &SLiveLinkHubRecordingListView::GetAssetContextMenu);
+			AssetPickerConfig.RefreshAssetViewDelegates.Add(&RefreshAssetViewDelegate);
 
 			if (AssetData.IsSet())
 			{
@@ -300,6 +302,42 @@ private:
 					FCanExecuteAction::CreateLambda([] () { return true; })
 				)
 			);
+
+			MenuBuilder.AddMenuEntry(
+				LOCTEXT("DeleteRecordingLabel", "Delete"),
+				LOCTEXT("DeleteRecordingTooltip", "Delete the recording"),
+				FSlateIcon(FAppStyle::GetAppStyleSetName(), "GenericCommands.Delete"),
+				FUIAction(
+					FExecuteAction::CreateLambda([SelectedAssetObject, this]()
+					{
+						if (SelectedAssetObject.IsValid())
+						{
+							IAssetRegistry& AssetRegistry = FModuleManager::LoadModuleChecked<FAssetRegistryModule>("AssetRegistry").Get();
+
+							TArray<FAssetData> Assets;
+							AssetRegistry.GetAssetsByPackageName(*SelectedAssetObject->GetPackage()->GetPathName(), Assets);
+						
+							if (TSharedPtr<FLiveLinkHubPlaybackController> PlaybackController = FLiveLinkHub::Get()->GetPlaybackController())
+							{
+								const TStrongObjectPtr<ULiveLinkRecording> PlaybackRecording = PlaybackController->GetRecording();
+								if (PlaybackRecording.Get() == SelectedAssetObject)
+								{
+									/** Delete asset upon ejection being completed. */
+									PlaybackController->Eject([Assets]() 
+									{
+										ObjectTools::DeleteAssets(Assets, false);
+									});
+								}
+								else
+								{
+									ObjectTools::DeleteAssets(Assets, false);
+								}
+							}
+						}
+					}),
+					FCanExecuteAction::CreateLambda([]() { return true; })
+				)
+			);
 			
 			MenuBuilder.AddMenuEntry(
 				LOCTEXT("DuplicateRecordingLabel", "Duplicate"),
@@ -398,6 +436,41 @@ private:
 					FCanExecuteAction::CreateLambda([SelectedAssetObject] () { return SelectedAssetObject.IsValid(); })
 				)
 			);
+
+			MenuBuilder.AddMenuEntry(
+				LOCTEXT("UpgradeRecordingInPlace", "Upgrade Recording in place"),
+				LOCTEXT("UpgradeRecordingInPlaceTooltip", "Loads the entire recording into memory, converts it to the latest version, and overwrites the recording with the updated version."),
+				FSlateIcon(FAppStyle::GetAppStyleSetName(), "Icons.Info"),
+				FUIAction(
+					FExecuteAction::CreateLambda([SelectedAssetObject, this]()
+					{
+						if (SelectedAssetObject.IsValid())
+						{
+							if (TSharedPtr<FLiveLinkHubPlaybackController> PlaybackController = FLiveLinkHub::Get()->GetPlaybackController())
+							{
+								const TStrongObjectPtr<ULiveLinkRecording> PlaybackRecording = PlaybackController->GetRecording();
+								if (PlaybackRecording.Get() == SelectedAssetObject)
+								{
+									/** Upgrade asset once it's ejection is complete. */
+									PlaybackController->Eject([SelectedAssetObject, this]()
+									{
+										constexpr bool bUpgradeInPlace = true;
+										FLiveLinkHub::Get()->UpgradeAndSaveRecording(CastChecked<ULiveLinkUAssetRecording>(SelectedAssetObject.Get()), bUpgradeInPlace);
+										RefreshAssetViewDelegate.ExecuteIfBound(true);
+									});
+								}
+								else
+								{
+									constexpr bool bUpgradeInPlace = true;
+									FLiveLinkHub::Get()->UpgradeAndSaveRecording(CastChecked<ULiveLinkUAssetRecording>(SelectedAssetObject.Get()), bUpgradeInPlace);
+									RefreshAssetViewDelegate.ExecuteIfBound(true);
+								}
+							}
+						}
+					}),
+					FCanExecuteAction::CreateLambda([SelectedAssetObject]() { return SelectedAssetObject.IsValid(); })
+				)
+			);
 		}
 		MenuBuilder.EndSection();
 
@@ -439,6 +512,8 @@ private:
 	TSharedPtr<SBox> BoxWidget;
 	/** Timer handle used for triggering a rename after duplicating a recording. */
 	FTimerHandle TimerHandle;
+	/** Delegate used to refresh the asset view. */
+	FRefreshAssetViewDelegate RefreshAssetViewDelegate;
 	/** True if there are recording assets that exist. */
 	mutable TOptional<bool> bAssetsAvailableCached;
 };
