LiveLinkHub - Implement compression on LiveLinkHub recordings and bump recording version

#Synced-CL 47110913

diff --git a/Engine/Plugins/Animation/LiveLinkHub/Source/LiveLinkHub/Private/Recording/Implementations/LiveLinkFrameData.h b/Engine/Plugins/Animation/LiveLinkHub/Source/LiveLinkHub/Private/Recording/Implementations/LiveLinkFrameData.h
--- a/Engine/Plugins/Animation/LiveLinkHub/Source/LiveLinkHub/Private/Recording/Implementations/LiveLinkFrameData.h
+++ b/Engine/Plugins/Animation/LiveLinkHub/Source/LiveLinkHub/Private/Recording/Implementations/LiveLinkFrameData.h
@@ -118,6 +118,8 @@ namespace UE::LiveLinkHub::FrameData::Private
 		FFrameBufferCache BufferedCache;
 		/** Current iteration data while buffering. */
 		FFrameBufferIterationData BufferIterationData;
+		/** Name of the compression method used on the frame data. */
+		FName CompressionMethodName;
 
 		/**
 		 * Find the closest frame matching the given timestamp.
diff --git a/Engine/Plugins/Animation/LiveLinkHub/Source/LiveLinkHub/Private/Recording/Implementations/LiveLinkUAssetRecording.cpp b/Engine/Plugins/Animation/LiveLinkHub/Source/LiveLinkHub/Private/Recording/Implementations/LiveLinkUAssetRecording.cpp
--- a/Engine/Plugins/Animation/LiveLinkHub/Source/LiveLinkHub/Private/Recording/Implementations/LiveLinkUAssetRecording.cpp
+++ b/Engine/Plugins/Animation/LiveLinkHub/Source/LiveLinkHub/Private/Recording/Implementations/LiveLinkUAssetRecording.cpp
@@ -5,6 +5,7 @@
 #include "Engine/Engine.h"
 #include "Features/IModularFeatures.h"
 #include "HAL/IConsoleManager.h"
+#include "LiveLinkCompression.h"
 #include "LiveLinkHubLog.h"
 #include "LiveLinkHubModule.h"
 #include "LiveLinkUAssetRecordingPlayer.h"
@@ -13,6 +14,7 @@
 #include "Modules/ModuleManager.h"
 #include "Recording/LiveLinkHubPlaybackController.h"
 #include "Recording/LiveLinkRecordingRangeHelpers.h"
+#include "Roles/LiveLinkAnimationTypes.h"
 #include "Serialization/LargeMemoryReader.h"
 #include "Serialization/ObjectAndNameAsStringProxyArchive.h"
 #include "UObject/Package.h"
@@ -21,6 +23,44 @@
 
 #define LIVELINKHUB_FRAME_DEBUG 0
 
+namespace UE::LiveLinkHub::Compression
+{
+	static const FName CompressionAlgo = NAME_Oodle;
+
+	/** Compress data with a bias for memory. */
+	void Compress(TArray<uint8>& InBytes)
+	{
+		TRACE_CPUPROFILER_EVENT_SCOPE(ULiveLinkUAssetRecording::Compress);
+
+		TArray<uint8> OutCompressedData;
+		ECompressionFlags CompressFlags = COMPRESS_BiasMemory;
+		int32 CompressedSize = FCompression::CompressMemoryBound(CompressionAlgo, InBytes.Num(), CompressFlags);
+		OutCompressedData.SetNumUninitialized(CompressedSize);
+
+		if (FCompression::CompressMemory(CompressionAlgo, OutCompressedData.GetData(), CompressedSize, InBytes.GetData(), InBytes.Num(), CompressFlags))
+		{
+			OutCompressedData.SetNum(CompressedSize, EAllowShrinking::No);
+			InBytes = MoveTemp(OutCompressedData);
+		}
+		else
+		{
+			ensureAlways(false);
+		}
+	}
+
+	/** Decompress data. */
+	TArray<uint8> Decompress(FName InCompressionMethod, TArray<uint8>& InBytes, int32 UncompressedSize)
+	{
+		TRACE_CPUPROFILER_EVENT_SCOPE(ULiveLinkUAssetRecording::Decompress);
+
+		TArray<uint8> UncompressedData;
+		UncompressedData.SetNumUninitialized(UncompressedSize);
+
+		ensure(FCompression::UncompressMemory(InCompressionMethod, UncompressedData.GetData(), UncompressedData.Num(), InBytes.GetData(), InBytes.Num(), COMPRESS_BiasMemory));
+		return UncompressedData;
+	}
+}
+
 ULiveLinkUAssetRecording::~ULiveLinkUAssetRecording()
 {
 	if (!IsTemplate())
@@ -552,7 +592,10 @@ void ULiveLinkUAssetRecording::SaveFrameData(FArchive* InFileWriter, const FLive
 
 	// Store all timestamps in a block.
 	*InFileWriter << InBaseDataContainer.Timestamps;
-	
+
+	FString CompressionMethod = UE::LiveLinkHub::Compression::CompressionAlgo.ToString();
+	*InFileWriter << CompressionMethod;
+
 	// Write the frame header size.
 	{
 		const uint64 CurrentPosition = InFileWriter->Tell();
@@ -564,7 +607,10 @@ void ULiveLinkUAssetRecording::SaveFrameData(FArchive* InFileWriter, const FLive
 	}
 
 	int64 RelativeStartPosition = 0;
-	
+
+	// Memory used for the intermediate step of serializing data before compressing it.
+	TArray<uint8> Backing;
+
 	for (int32 FrameIdx = 0; FrameIdx < NumFrames; ++FrameIdx)
 	{
 		TSharedPtr<FInstancedStruct>& Frame = InBaseDataContainer.RecordedData[FrameIdx];
@@ -576,10 +622,22 @@ void ULiveLinkUAssetRecording::SaveFrameData(FArchive* InFileWriter, const FLive
 		// Write the frame index for streaming frames when loading.
 		*InFileWriter << FrameIdx;
 			
-		// Write the entire frame data.
-		FObjectAndNameAsStringProxyArchive StructAr(*InFileWriter, false);
+		// Write the entire frame data to memory
+		FMemoryWriter MemoryWriter{ Backing, true };
+		FObjectAndNameAsStringProxyArchive StructAr(MemoryWriter, false);
+
 		Frame->Serialize(StructAr);
 
+		int32 UncompressedSize = Backing.Num();
+		*InFileWriter << UncompressedSize;
+
+		// Then compress it and write it to the file. 
+		// Note: We compress every frame independently to make it easier for the loading code to stream the data.
+		UE::LiveLinkHub::Compression::Compress(Backing);
+		InFileWriter->Serialize(Backing.GetData(), Backing.Num());
+
+		Backing.Reset();
+
 		// Store the serialized frame size, so we can write it once later.
 		{
 			const int64 EndFramePosition = InFileWriter->Tell();
@@ -683,7 +741,7 @@ void ULiveLinkUAssetRecording::LoadRecordingAsync(int32 InStartFrame, int32 InCu
 				bCancelStream = true;
 				return;
 			}
-				
+
 			FLiveLinkRecordingStaticDataContainer& DataContainer = RecordingData.StaticData.FindChecked(*TemporaryFrameData.FrameDataSubjectKey.Get());
 			LoadFrameData(TemporaryFrameData, DataContainer, 0, 0, TemporaryFrameData.MaxFrames, /* bForceSequential */ true);
 
@@ -892,6 +950,14 @@ bool ULiveLinkUAssetRecording::LoadInitialFrameData(UE::LiveLinkHub::FrameData::
 		{
 			Reader->GetMemoryReader() << OutFrameData.Timestamps;
 		}
+
+		if (RecordingVersionBeingLoaded >= UE::LiveLinkHub::Private::RecordingVersions::CompressionSupport)
+		{
+			FString CompressionMethodName;
+			Reader->GetMemoryReader() << CompressionMethodName;
+
+			OutFrameData.CompressionMethodName = *CompressionMethodName;
+		}
 	}
 
 	OutFrameData.LastTimestamp = OutFrameData.Timestamps.Num () > 0 ? OutFrameData.Timestamps.Last() : 0.0;
@@ -1206,6 +1272,7 @@ void ULiveLinkUAssetRecording::LoadFrameData(UE::LiveLinkHub::FrameData::Private
 			RawFramesRightBatch.Count = FMath::Max(RightBatchSize, 0);
 		}
 
+		TRACE_CPUPROFILER_EVENT_SCOPE(ULiveLinkUAssetRecording::LoadFrames)
 		// Iterate through the entire range, loading frames from cache or disk. This runs until a batch cycle has completed
 		// or the entire range is loaded.
 		int32 FramesLoaded = 0;
@@ -1362,6 +1429,7 @@ bool ULiveLinkUAssetRecording::LoadFrameFromDisk(const int32 InFrame,
 	TSharedPtr<FInstancedStruct>& OutFrame, double& OutTimestamp,
 	const TSharedPtr<FLiveLinkHubBulkData::FScopedBulkDataMemoryReader>& InMemory)
 {
+	TRACE_CPUPROFILER_EVENT_SCOPE(ULiveLinkUAssetRecording::LoadFrameFromDisk);
 	DECLARE_SCOPE_CYCLE_COUNTER(TEXT("ULiveLinkUAssetRecording::LoadFrameFromDisk"), STAT_ULiveLinkUAssetRecording_LoadFrameFromDisk, STATGROUP_LiveLinkHub);
 	OutTimestamp = 0.0;
 
@@ -1411,9 +1479,45 @@ bool ULiveLinkUAssetRecording::LoadFrameFromDisk(const int32 InFrame,
 	{
 		DECLARE_SCOPE_CYCLE_COUNTER(TEXT("ULiveLinkUAssetRecording::LoadFrameFromDisk::SerializeFrame"),
 			STAT_ULiveLinkUAssetRecording_LoadFrameFromDisk_SerializeFrame, STATGROUP_LiveLinkHub);
-		FObjectAndNameAsStringProxyArchive StructAr(Reader->GetMemoryReader(), true);
-		OutFrame = MakeShared<FInstancedStruct>(LoadedStruct);
-		OutFrame->Serialize(StructAr);
+
+
+		if (RecordingVersionBeingLoaded < UE::LiveLinkHub::Private::RecordingVersions::CompressionSupport)
+		{
+			FObjectAndNameAsStringProxyArchive StructAr(Reader->GetMemoryReader(), true);
+			OutFrame = MakeShared<FInstancedStruct>(LoadedStruct);
+			OutFrame->Serialize(StructAr);
+		}
+		else
+		{
+			int32 FrameSize = InFrameData.FrameDiskSizes[InFrame].Get<1>();
+			int32 UncompressedSize = 0;
+
+			Reader->GetMemoryReader() << UncompressedSize;
+
+			const int32 HeaderSize = sizeof(FrameSize) + sizeof(UncompressedSize);
+
+			// We've already read the header, so substract that amount from the frame size we're going to read next.
+			FrameSize -= HeaderSize;
+			check(FrameSize > 0);
+
+			TArray<uint8> CompressedData;
+			CompressedData.AddZeroed(FrameSize);
+
+			{
+				TRACE_CPUPROFILER_EVENT_SCOPE(ULiveLinkUAssetRecording::ReadCompressedData);
+				Reader->GetMemoryReader().Serialize(CompressedData.GetData(), FrameSize);
+			}
+
+			TArray<uint8> UncompressedData = UE::LiveLinkHub::Compression::Decompress(InFrameData.CompressionMethodName, CompressedData, UncompressedSize);
+
+			FMemoryReader UncompressedReader{ UncompressedData };
+
+			FObjectAndNameAsStringProxyArchive StructAr(UncompressedReader, true);
+			OutFrame = MakeShared<FInstancedStruct>(LoadedStruct);
+
+			TRACE_CPUPROFILER_EVENT_SCOPE(ULiveLinkUAssetRecording::DeserializeFrame);
+			OutFrame->Serialize(StructAr);
+		}
 	}
 	
 	OutTimestamp = Timestamp;
diff --git a/Engine/Plugins/Animation/LiveLinkHub/Source/LiveLinkHub/Private/Recording/LiveLinkHubRecordingVersions.h b/Engine/Plugins/Animation/LiveLinkHub/Source/LiveLinkHub/Private/Recording/LiveLinkHubRecordingVersions.h
--- a/Engine/Plugins/Animation/LiveLinkHub/Source/LiveLinkHub/Private/Recording/LiveLinkHubRecordingVersions.h
+++ b/Engine/Plugins/Animation/LiveLinkHub/Source/LiveLinkHub/Private/Recording/LiveLinkHubRecordingVersions.h
@@ -11,7 +11,9 @@ namespace UE::LiveLinkHub::Private::RecordingVersions
 	constexpr int32 DynamicFrameSizes = 2;
 	/** Timestamps serialized in bulk. */
 	constexpr int32 BulkTimestamps = 3;
+	/** Added oodle compression to serialized frame data. */
+	constexpr int32 CompressionSupport = 4;
 
 	/** The latest recording version. UPDATE IF ADDING NEW VERSIONS! */
-	constexpr int32 Latest = BulkTimestamps;
+	constexpr int32 Latest = CompressionSupport;
 }
