LiveLinkHub - Re-add topology mode switcher widget that was accidentally removed

#rnx

diff --git a/Engine/Plugins/Animation/LiveLinkHub/Source/LiveLinkHub/Private/LiveLinkHubApplicationMode.cpp b/Engine/Plugins/Animation/LiveLinkHub/Source/LiveLinkHub/Private/LiveLinkHubApplicationMode.cpp
--- a/Engine/Plugins/Animation/LiveLinkHub/Source/LiveLinkHub/Private/LiveLinkHubApplicationMode.cpp
+++ b/Engine/Plugins/Animation/LiveLinkHub/Source/LiveLinkHub/Private/LiveLinkHubApplicationMode.cpp
@@ -45,6 +45,24 @@ void FLiveLinkHubApplicationMode::PostActivateMode()
 
 	UToolMenu* Menu = UToolMenus::Get()->ExtendMenu(TEXT("LiveLinkHub.ToolBar"));
 
+	FToolMenuSection& GlobalWidgetsSection = Menu->AddSection(TEXT("GlobalWidgets"), FText::GetEmpty(), FToolMenuInsert(NAME_None, EToolMenuInsertType::First));
+
+	TSharedRef<SWidget> GlobalWidget =
+		SNew(SHorizontalBox)
+		+ SHorizontalBox::Slot()
+		.AutoWidth()
+		.Padding(2.0, 0.0)
+		.VAlign(VAlign_Center)
+		[
+			SNew(SLiveLinkHubTopologyModeSwitcher)
+		];
+
+	GlobalWidgetsSection.AddEntry(
+		FToolMenuEntry::InitWidget(TEXT("TestWidget"), GlobalWidget, FText::GetEmpty(), true, false)
+	);
+
+	GlobalWidgetsSection.AddSeparator(NAME_None);
+
 	FToolMenuSection& SessionConfigSection = Menu->AddSection(TEXT("ModeWidgets"), FText::GetEmpty(), FToolMenuInsert(TEXT("GlobalWidgets"), EToolMenuInsertType::After));
 
 	TSharedRef<SHorizontalBox> ModeWidgets = SNew(SHorizontalBox);
