LiveLinkHub - Fix crash when playing back recording after upgrading it

#Synced-CL 48456341

diff --git a/Engine/Plugins/Animation/LiveLinkHub/Source/LiveLinkHub/Private/LiveLinkHub.cpp b/Engine/Plugins/Animation/LiveLinkHub/Source/LiveLinkHub/Private/LiveLinkHub.cpp
--- a/Engine/Plugins/Animation/LiveLinkHub/Source/LiveLinkHub/Private/LiveLinkHub.cpp
+++ b/Engine/Plugins/Animation/LiveLinkHub/Source/LiveLinkHub/Private/LiveLinkHub.cpp
@@ -581,12 +581,12 @@ void FLiveLinkHub::OnPlaybackReadyForUpgrade(bool bUpgradeInPlace)
 	ULiveLinkUAssetRecording* CurrentRecording = CastChecked<ULiveLinkUAssetRecording>(PlaybackController->GetRecording().Get());
 		
 	TWeakObjectPtr<ULiveLinkUAssetRecording> CurrentRecordingWeakPtr(CurrentRecording);
-	CurrentRecording->GetOnRecordingFullyLoaded().BindLambda([this, CurrentRecordingWeakPtr, RecordingPath]()
+	CurrentRecording->GetOnRecordingFullyLoaded().BindLambda([this, CurrentRecordingWeakPtr, RecordingPath, bUpgradeInPlace]()
 	{
 		const TStrongObjectPtr<ULiveLinkUAssetRecording> CurrentRecordingPin = CurrentRecordingWeakPtr.Pin();
 		if (CurrentRecordingPin.IsValid() && RecordingController)
 		{
-			RecordingController->RecordFromExistingAndSave(CurrentRecordingPin.Get(), RecordingPath);
+			RecordingController->RecordFromExistingAndSave(CurrentRecordingPin.Get(), RecordingPath, bUpgradeInPlace);
 		}
 	});
 
diff --git a/Engine/Plugins/Animation/LiveLinkHub/Source/LiveLinkHub/Private/Recording/Implementations/LiveLinkUAssetRecorder.cpp b/Engine/Plugins/Animation/LiveLinkHub/Source/LiveLinkHub/Private/Recording/Implementations/LiveLinkUAssetRecorder.cpp
--- a/Engine/Plugins/Animation/LiveLinkHub/Source/LiveLinkHub/Private/Recording/Implementations/LiveLinkUAssetRecorder.cpp
+++ b/Engine/Plugins/Animation/LiveLinkHub/Source/LiveLinkHub/Private/Recording/Implementations/LiveLinkUAssetRecorder.cpp
@@ -143,11 +143,19 @@ bool FLiveLinkUAssetRecorder::CanRecord(FText* OutErrorMessage) const
 	return true;
 }
 
-void FLiveLinkUAssetRecorder::SaveFromExistingRecording(const ULiveLinkRecording* InRecording, const FString& InFilePath)
+void FLiveLinkUAssetRecorder::SaveFromExistingRecording(ULiveLinkRecording* InRecording, const FString& InFilePath, bool bUpgradeInPlace)
 {
-	FLiveLinkUAssetRecordingData InRecordingData = CastChecked<ULiveLinkUAssetRecording>(InRecording)->RecordingData;
-	SaveRecording(MoveTemp(InRecordingData), InRecording->LengthInSeconds, /*bIncrementTake*/ false,
-		InRecording->RecordingPreset, &InRecording->FrameRate, InFilePath);
+	ULiveLinkUAssetRecording* UAssetRecording = CastChecked<ULiveLinkUAssetRecording>(InRecording);
+
+	if (bUpgradeInPlace)
+	{
+		SaveRecordingInPlace(UAssetRecording);
+	}
+	else
+	{
+		SaveRecording(MoveTemp(UAssetRecording->RecordingData), InRecording->LengthInSeconds, /*bIncrementTake*/ false,
+			InRecording->RecordingPreset, &InRecording->FrameRate, InFilePath);
+	}
 
 	CurrentRecording.Reset();
 }
@@ -326,17 +334,8 @@ void FLiveLinkUAssetRecorder::SaveRecording(FLiveLinkUAssetRecordingData&& InRec
 	UPackage* Package = FindPackage(NULL, *PackageName);
 	ULiveLinkUAssetRecording* Recording = nullptr;
 
-	if (Package)
-	{
-		// If the package already exists, then that means that we shouldn't create a new asset, and just overwrite the existing recording.
-		Recording = Cast<ULiveLinkUAssetRecording>(StaticFindObject(ULiveLinkUAssetRecording::StaticClass(), FTopLevelAssetPath(*PackageName, *NewAssetName), EFindObjectFlags::None));
-		check(Recording);
-	}
-	else
-	{
-		Package = CreatePackage(*PackageName);
-		Recording = NewObject<ULiveLinkUAssetRecording>(Package, *NewAssetName, RF_Public | RF_Standalone);
-	}
+	Package = CreatePackage(*PackageName);
+	Recording = NewObject<ULiveLinkUAssetRecording>(Package, *NewAssetName, RF_Public | RF_Standalone);
 
 	if (Recording)
 	{
@@ -348,6 +347,7 @@ void FLiveLinkUAssetRecorder::SaveRecording(FLiveLinkUAssetRecordingData&& InRec
 			// such as during an upgrade.
 			UEngine::CopyPropertiesForUnrelatedObjects(InPreset, Recording->RecordingPreset);
 		}
+
 		if (InFrameRate)
 		{
 			Recording->FrameRate = *InFrameRate;
@@ -373,6 +373,25 @@ void FLiveLinkUAssetRecorder::SaveRecording(FLiveLinkUAssetRecordingData&& InRec
 	}
 }
 
+void FLiveLinkUAssetRecorder::SaveRecordingInPlace(ULiveLinkUAssetRecording* InRecording)
+{
+	if (!InRecording)
+	{
+		return;
+	}
+
+	InRecording->SavedRecordingVersion = UE::LiveLinkHub::Private::RecordingVersions::Latest;
+	InRecording->MarkPackageDirty();
+
+	const TStrongObjectPtr<ULiveLinkUAssetRecording> NewRecordingStrongPtr(InRecording);
+
+	// Create a task to run on a separate thread for saving all frame data and writing the final UAsset to disk.
+	// We use a container rather than just one task on the chance a save operation is still running when another recording is being saved.
+	FAsyncTask<FLiveLinkSaveRecordingAsyncTask>& AsyncTask = *AsyncSaveTasks.Add(NewRecordingStrongPtr,
+		MakeUnique<FAsyncTask<FLiveLinkSaveRecordingAsyncTask>>(InRecording, this));
+	AsyncTask.StartBackgroundTask();
+}
+
 void FLiveLinkUAssetRecorder::RecordInitialStaticData()
 {
 	FLiveLinkHubClient* LiveLinkClient = static_cast<FLiveLinkHubClient*>(&IModularFeatures::Get().GetModularFeature<ILiveLinkClient>(ILiveLinkClient::ModularFeatureName));
diff --git a/Engine/Plugins/Animation/LiveLinkHub/Source/LiveLinkHub/Private/Recording/Implementations/LiveLinkUAssetRecorder.h b/Engine/Plugins/Animation/LiveLinkHub/Source/LiveLinkHub/Private/Recording/Implementations/LiveLinkUAssetRecorder.h
--- a/Engine/Plugins/Animation/LiveLinkHub/Source/LiveLinkHub/Private/Recording/Implementations/LiveLinkUAssetRecorder.h
+++ b/Engine/Plugins/Animation/LiveLinkHub/Source/LiveLinkHub/Private/Recording/Implementations/LiveLinkUAssetRecorder.h
@@ -25,7 +25,7 @@ public:
 	virtual void RecordFrameData(const FLiveLinkSubjectKey& SubjectKey, const FLiveLinkFrameDataStruct& FrameData) override;
 	virtual bool IsRecording() const override;
 	virtual bool CanRecord(FText* OutErrorMessage) const override;
-	virtual void SaveFromExistingRecording(const ULiveLinkRecording* InRecording, const FString& InFilePath) override;
+	virtual void SaveFromExistingRecording(ULiveLinkRecording* InRecording, const FString& InFilePath, bool bUpgradeInPlace) override;
 	virtual bool GetSavePresetPackageName(FString& OutName, const FString& InDefaultName = FString(), bool bAlwaysOpenSaveAsDialog = false) const override;
 	//~ End ILiveLinkRecorder
 
@@ -36,6 +36,7 @@ private:
 	void SaveRecording(FLiveLinkUAssetRecordingData&& InRecordingData, double InRecordingLength,
 		bool bIncrementTake, const TObjectPtr<ULiveLinkPreset>& InPreset = nullptr, const FFrameRate* InFrameRate = nullptr,
 		const FString& InFilePath = FString());
+	void SaveRecordingInPlace(ULiveLinkUAssetRecording* InRecording);
 	/** Record data to a ULiveLinkRecording object. */
 	void RecordBaseData(FLiveLinkRecordingBaseDataContainer& StaticDataContainer, TSharedPtr<FInstancedStruct>&& DataToRecord);
 	/** Record initial data for all livelink subjects. (Useful when static data was sent before the recording started). */
diff --git a/Engine/Plugins/Animation/LiveLinkHub/Source/LiveLinkHub/Private/Recording/LiveLinkHubRecordingController.cpp b/Engine/Plugins/Animation/LiveLinkHub/Source/LiveLinkHub/Private/Recording/LiveLinkHubRecordingController.cpp
--- a/Engine/Plugins/Animation/LiveLinkHub/Source/LiveLinkHub/Private/Recording/LiveLinkHubRecordingController.cpp
+++ b/Engine/Plugins/Animation/LiveLinkHub/Source/LiveLinkHub/Private/Recording/LiveLinkHubRecordingController.cpp
@@ -79,9 +79,9 @@ void FLiveLinkHubRecordingController::RecordFrameData(const FLiveLinkSubjectKey&
 	RecorderImplementation->RecordFrameData(SubjectKey, FrameData);
 }
 
-void FLiveLinkHubRecordingController::RecordFromExistingAndSave(const ULiveLinkUAssetRecording* InRecording, const FString& InFilePath)
+void FLiveLinkHubRecordingController::RecordFromExistingAndSave(ULiveLinkRecording* InRecording, const FString& InFilePath, bool bUpgradeInPlace)
 {
-	RecorderImplementation->SaveFromExistingRecording(InRecording, InFilePath);
+	RecorderImplementation->SaveFromExistingRecording(InRecording, InFilePath, bUpgradeInPlace);
 }
 
 bool FLiveLinkHubRecordingController::OpenSaveAsPrompt(FString& OutName, const FString& InInitialName) const
diff --git a/Engine/Plugins/Animation/LiveLinkHub/Source/LiveLinkHub/Private/Recording/LiveLinkHubRecordingController.h b/Engine/Plugins/Animation/LiveLinkHub/Source/LiveLinkHub/Private/Recording/LiveLinkHubRecordingController.h
--- a/Engine/Plugins/Animation/LiveLinkHub/Source/LiveLinkHub/Private/Recording/LiveLinkHubRecordingController.h
+++ b/Engine/Plugins/Animation/LiveLinkHub/Source/LiveLinkHub/Private/Recording/LiveLinkHubRecordingController.h
@@ -41,7 +41,7 @@ public:
 	void RecordFrameData(const FLiveLinkSubjectKey& SubjectKey, const FLiveLinkFrameDataStruct& FrameData);
 
 	/** Record from an existing recording. */
-	void RecordFromExistingAndSave(const ULiveLinkUAssetRecording* InRecording, const FString& InFilePath);
+	void RecordFromExistingAndSave(ULiveLinkRecording* InRecording, const FString& InFilePath, bool bUpgradeInPlace = false);
 
 	/** Open a save as window and retrieve a file path. */
 	bool OpenSaveAsPrompt(FString& OutName, const FString& InInitialName) const;
diff --git a/Engine/Plugins/Animation/LiveLinkHub/Source/LiveLinkHub/Private/Recording/LiveLinkRecorder.h b/Engine/Plugins/Animation/LiveLinkHub/Source/LiveLinkHub/Private/Recording/LiveLinkRecorder.h
--- a/Engine/Plugins/Animation/LiveLinkHub/Source/LiveLinkHub/Private/Recording/LiveLinkRecorder.h
+++ b/Engine/Plugins/Animation/LiveLinkHub/Source/LiveLinkHub/Private/Recording/LiveLinkRecorder.h
@@ -27,7 +27,7 @@ public:
 	/** Returns whether we can begin recording. */
 	virtual bool CanRecord(FText* OutErrorMessage = nullptr) const = 0;
 	/** Directly provide the recording and save it. */
-	virtual void SaveFromExistingRecording(const ULiveLinkRecording* InRecording, const FString& InFilePath) = 0;
+	virtual void SaveFromExistingRecording(ULiveLinkRecording* InRecording, const FString& InFilePath, bool bUpgradeInPlace) = 0;
 	/** Creates a unique asset name and prompts the user for the recording name. */
 	virtual bool GetSavePresetPackageName(FString& OutName, const FString& InDefaultName, bool bAlwaysOpenSaveAsDialog = false) const = 0;
 };
