LiveLinkHub - Fix customization not showing for LiveLinkHubClientId

#Synced-CL 48735856

diff --git a/Engine/Plugins/Animation/LiveLinkHub/Source/LiveLinkHubEditor/Private/LiveLinkHubEditorModule.cpp b/Engine/Plugins/Animation/LiveLinkHub/Source/LiveLinkHubEditor/Private/LiveLinkHubEditorModule.cpp
--- a/Engine/Plugins/Animation/LiveLinkHub/Source/LiveLinkHubEditor/Private/LiveLinkHubEditorModule.cpp
+++ b/Engine/Plugins/Animation/LiveLinkHub/Source/LiveLinkHubEditor/Private/LiveLinkHubEditorModule.cpp
@@ -29,48 +29,48 @@ DECLARE_LOG_CATEGORY_CLASS(LogLiveLinkHubEditor, Log, Log)
 
 void FLiveLinkHubEditorModule::StartupModule()
 {
-	if (!IsRunningCommandlet() && CVarLiveLinkHubEnableStatusBar.GetValueOnAnyThread())
-	{
-		FCoreDelegates::OnPostEngineInit.AddRaw(this, &FLiveLinkHubEditorModule::OnPostEngineInit);
-	}
+	FCoreDelegates::OnPostEngineInit.AddRaw(this, &FLiveLinkHubEditorModule::OnPostEngineInit);
 }
 
 void FLiveLinkHubEditorModule::ShutdownModule()
 {
 	UToolMenus::UnregisterOwner(this);
 
+	FCoreDelegates::OnPostEngineInit.RemoveAll(this);
 	if (!IsRunningCommandlet() && CVarLiveLinkHubEnableStatusBar.GetValueOnAnyThread()) 
-	{
-		FCoreDelegates::OnPostEngineInit.RemoveAll(this);
+	{	
 		UnregisterLiveLinkHubStatusBar();
+	}
 
-		if (FPropertyEditorModule* PropertyModule = FModuleManager::GetModulePtr<FPropertyEditorModule>("PropertyEditor"))
-		{
-			PropertyModule->UnregisterCustomPropertyTypeLayout("LiveLinkHubClientId");
-		}
+	if (FPropertyEditorModule* PropertyModule = FModuleManager::GetModulePtr<FPropertyEditorModule>("PropertyEditor"))
+	{
+		PropertyModule->UnregisterCustomPropertyTypeLayout("LiveLinkHubClientId");
 	}
 }
 
 void FLiveLinkHubEditorModule::OnPostEngineInit()
 {
-	if (GEditor && FApp::CanEverRender())
+	if (!IsRunningCommandlet() && CVarLiveLinkHubEnableStatusBar.GetValueOnAnyThread())
 	{
-		RegisterLiveLinkHubStatusBar();
+		if (GEditor && FApp::CanEverRender())
+		{
+			RegisterLiveLinkHubStatusBar();
 		
-		FToolMenuOwnerScoped OwnerScoped(this);
-		UToolMenu* Menu = UToolMenus::Get()->ExtendMenu("LevelEditor.MainMenu.Tools");
-		FToolMenuSection& Section = Menu->AddSection("VirtualProductionSection", LOCTEXT("VirtualProductionSection", "Virtual Production"));
-
-		Section.AddMenuEntry("LiveLinkHub",
-			LOCTEXT("LiveLinkHubLabel", "Live Link Hub"),
-			LOCTEXT("LiveLinkHubTooltip", "Launch the Live Link Hub app."),
-			FSlateIcon("LiveLinkStyle", "LiveLinkClient.Common.Icon.Small"),
-			FUIAction(FExecuteAction::CreateStatic(&UE::LiveLinkHubLauncherUtils::OpenLiveLinkHub)));
-
-		FPropertyEditorModule& PropertyModule = FModuleManager::LoadModuleChecked<FPropertyEditorModule>("PropertyEditor");
-		PropertyModule.RegisterCustomPropertyTypeLayout("LiveLinkHubClientId", FOnGetPropertyTypeCustomizationInstance::CreateStatic(&FLiveLinkHubClientIdCustomization::MakeInstance));
-		PropertyModule.NotifyCustomizationModuleChanged();
+			FToolMenuOwnerScoped OwnerScoped(this);
+			UToolMenu* Menu = UToolMenus::Get()->ExtendMenu("LevelEditor.MainMenu.Tools");
+			FToolMenuSection& Section = Menu->AddSection("VirtualProductionSection", LOCTEXT("VirtualProductionSection", "Virtual Production"));
+
+			Section.AddMenuEntry("LiveLinkHub",
+				LOCTEXT("LiveLinkHubLabel", "Live Link Hub"),
+				LOCTEXT("LiveLinkHubTooltip", "Launch the Live Link Hub app."),
+				FSlateIcon("LiveLinkStyle", "LiveLinkClient.Common.Icon.Small"),
+				FUIAction(FExecuteAction::CreateStatic(&UE::LiveLinkHubLauncherUtils::OpenLiveLinkHub)));
+		}
 	}
+
+	FPropertyEditorModule& PropertyModule = FModuleManager::LoadModuleChecked<FPropertyEditorModule>("PropertyEditor");
+	PropertyModule.RegisterCustomPropertyTypeLayout("LiveLinkHubClientId", FOnGetPropertyTypeCustomizationInstance::CreateStatic(&FLiveLinkHubClientIdCustomization::MakeInstance));
+	PropertyModule.NotifyCustomizationModuleChanged();
 }
 
 void FLiveLinkHubEditorModule::RegisterLiveLinkHubStatusBar()
