LiveLinkHub: Replace stale LLH instances when a new one from the same machine reconnects.

This avoids the issue of stale LLH sources in Unreal, which can end up with the stale subjects trying to drive the Live Link consumers, or the correct subjects being disabled instead of the stale ones.

Also added a log when RemoveAllSources is called but the LiveLinkProvider is unexpectedly invalid.  This also makes it so discovery messages include annotations, which was missing from the MB LL Publish API (only SendMessage supported annotations).

#rnx

#Synced-CL 49070374

diff --git a/Engine/Plugins/Animation/LiveLinkHub/Source/LiveLinkHub/Private/Clients/LiveLinkHubProvider.cpp b/Engine/Plugins/Animation/LiveLinkHub/Source/LiveLinkHub/Private/Clients/LiveLinkHubProvider.cpp
--- a/Engine/Plugins/Animation/LiveLinkHub/Source/LiveLinkHub/Private/Clients/LiveLinkHubProvider.cpp
+++ b/Engine/Plugins/Animation/LiveLinkHub/Source/LiveLinkHub/Private/Clients/LiveLinkHubProvider.cpp
@@ -341,7 +341,7 @@ void FLiveLinkHubProvider::AddRestoredClient(FLiveLinkHubUEClientInfo& RestoredC
 				UE::TWriteScopeLock Locker(ClientsMapLock);
 				ClientsMap.FindOrAdd(RestoredClientInfo.Id) = RestoredClientInfo;
 
-				SendMessage(MakeLiveLinkHubDiscoveryMessage(), FoundClient->ControlEndpoint, EMessageFlags::Reliable);
+				SendMessage(MakeLiveLinkHubDiscoveryMessage(), FoundClient->ControlEndpoint, EMessageFlags::Reliable, GetAnnotations());
 				UE_LOG(LogLiveLinkHub, Log, TEXT("Restoring connection to client %s (%s %s)"), *FoundClient->ClientId.ToString(), *FoundClient->Hostname, *FoundClient->IP);
 			}
 			else
@@ -668,9 +668,9 @@ void FLiveLinkHubProvider::ConnectToAllDiscoveredClients()
 	}
 
 	// For Discovery Protocol V1
-	Publish(MakeLiveLinkHubDiscoveryMessage());
+	Publish(MakeLiveLinkHubDiscoveryMessage(), GetAnnotations());
 
-	SendMessage(MakeLiveLinkHubDiscoveryMessage(), CompatibleClients, EMessageFlags::Reliable);
+	SendMessage(MakeLiveLinkHubDiscoveryMessage(), CompatibleClients, EMessageFlags::Reliable, GetAnnotations());
 }
 
 void FLiveLinkHubProvider::OnConnectionsClosed(const TArray<FMessageAddress>& ClosedAddresses)
@@ -765,7 +765,7 @@ void FLiveLinkHubProvider::ValidateHubConnections()
 							It->bSentDiscoveryRequest = true;
 
 							UE_LOG(LogLiveLinkHub, Verbose, TEXT("Sending Hub Connection request to %s"), *It->Hostname);
-							SendMessage(MakeLiveLinkHubDiscoveryMessage(), It->ControlEndpoint, EMessageFlags::Reliable);
+							SendMessage(MakeLiveLinkHubDiscoveryMessage(), It->ControlEndpoint, EMessageFlags::Reliable, GetAnnotations());
 						}
 					}
 				}
@@ -980,7 +980,7 @@ void FLiveLinkHubProvider::ConnectTo(FLiveLinkHubClientId Client)
 	}
 	else
 	{
-		SendMessage(MakeLiveLinkHubDiscoveryMessage(), Entry.ControlEndpoint);
+		SendMessage(MakeLiveLinkHubDiscoveryMessage(), Entry.ControlEndpoint, EMessageFlags::None, GetAnnotations());
 	}
 }
 
@@ -1316,7 +1316,7 @@ void FLiveLinkHubProvider::OnFiltersModified()
 				EntryIt->bDisableAutoconnect = false;
 
 				UE_LOG(LogLiveLinkHub, Verbose, TEXT("Sending Hub Connection request to %s"), *EntryIt->Hostname);
-				SendMessage(MakeLiveLinkHubDiscoveryMessage(), EntryIt->ControlEndpoint, EMessageFlags::Reliable);
+				SendMessage(MakeLiveLinkHubDiscoveryMessage(), EntryIt->ControlEndpoint, EMessageFlags::Reliable, GetAnnotations());
 			}
 		}
 	}
diff --git a/Engine/Plugins/Animation/LiveLinkHub/Source/LiveLinkHub/Private/Session/LiveLinkHubSession.h b/Engine/Plugins/Animation/LiveLinkHub/Source/LiveLinkHub/Private/Session/LiveLinkHubSession.h
--- a/Engine/Plugins/Animation/LiveLinkHub/Source/LiveLinkHub/Private/Session/LiveLinkHubSession.h
+++ b/Engine/Plugins/Animation/LiveLinkHub/Source/LiveLinkHub/Private/Session/LiveLinkHubSession.h
@@ -187,6 +187,10 @@ public:
 				LiveLinkProvider->DisconnectClient(ClientId);
 			}
 		}
+		else
+		{
+			UE_LOG(LogLiveLinkHub, Warning, TEXT("LiveLinkProvider was already destroyed when RemoveAllClients was called"));
+		}
 
 		{
 			FWriteScopeLock Locker(SessionDataLock);
diff --git a/Engine/Plugins/Animation/LiveLinkHub/Source/LiveLinkHubMessaging/Private/LiveLinkHubControlChannel.h b/Engine/Plugins/Animation/LiveLinkHub/Source/LiveLinkHubMessaging/Private/LiveLinkHubControlChannel.h
--- a/Engine/Plugins/Animation/LiveLinkHub/Source/LiveLinkHubMessaging/Private/LiveLinkHubControlChannel.h
+++ b/Engine/Plugins/Animation/LiveLinkHub/Source/LiveLinkHubMessaging/Private/LiveLinkHubControlChannel.h
@@ -6,6 +6,7 @@
 #include "LiveLinkHubConnectionManager.h"
 #include "ILiveLinkHubMessagingModule.h"
 #include "LiveLinkHubMessageBusSource.h"
+#include "LiveLinkHubMessages.h"
 #include "MessageEndpoint.h"
 #include "MessageEndpointBuilder.h"
 #include "Misc/EngineVersion.h"
@@ -224,6 +225,61 @@ private:
 #endif
 	}
 
+	/**
+	 * Checks if an incoming LiveLinkHub discovery should replace existing sources.
+	 * Returns the GUID of any existing source that should be removed, or invalid GUID if none.
+	 * This method is only called after confirming the incoming connection is a LiveLinkHub provider.
+	 */
+	FGuid FindStaleLiveLinkHubSource(const FLiveLinkHubDiscoveryMessage& IncomingMessage, const TSharedRef<IMessageContext, ESPMode::ThreadSafe>& Context) const
+	{
+		ILiveLinkClient* Client = IModularFeatures::Get().IsModularFeatureAvailable(ILiveLinkClient::ModularFeatureName)
+			? &IModularFeatures::Get().GetModularFeature<ILiveLinkClient>(ILiveLinkClient::ModularFeatureName)
+			: nullptr;
+
+		if (!Client)
+		{
+			return FGuid();
+		}
+
+		// This method is only called for confirmed LiveLinkHub providers
+		// Search existing sources for LiveLinkHub instances from the same machine
+		TArray<FGuid> Sources = Client->GetSources();
+
+		for (const FGuid& SourceId : Sources)
+		{
+			FText SourceType = Client->GetSourceType(SourceId);
+			FText SourceMachineName = Client->GetSourceMachineName(SourceId);
+
+			// Look for existing LiveLinkHub sources from the same machine
+			// We check both the source type and machine name to identify potential duplicates
+			bool bIsLiveLinkHubSource = SourceType.ToString().Contains(TEXT("Live Link Hub")) ||
+			                           SourceType.ToString().Contains(UE::LiveLinkHub::Private::LiveLinkHubProviderType.ToString());
+
+			if (bIsLiveLinkHubSource && SourceMachineName.ToString() == IncomingMessage.MachineName)
+			{
+				// Check if the source is still valid
+				bool bIsStillValid = Client->IsSourceStillValid(SourceId);
+
+				if (!bIsStillValid)
+				{
+					// Source is invalid (disconnected/timed out), replace it
+					UE_LOG(LogLiveLinkHubMessaging, Verbose, TEXT("Found invalid LiveLinkHub source %s from %s, will replace"),
+						*SourceId.ToString(), *IncomingMessage.MachineName);
+					return SourceId;
+				}
+				else if (SourceType.ToString() == IncomingMessage.ProviderName)
+				{
+					// Same provider name from same machine - likely a reconnection
+					UE_LOG(LogLiveLinkHubMessaging, Verbose, TEXT("Found duplicate LiveLinkHub source %s from %s, will replace"),
+						*SourceId.ToString(), *IncomingMessage.MachineName);
+					return SourceId;
+				}
+			}
+		}
+
+		return FGuid(); // No stale source found
+	}
+
 	/** Templated method to dispatch messages using template specialization. */
 	template <typename MessageType>
 	void Handle(const MessageType& Message, const TSharedRef<IMessageContext, ESPMode::ThreadSafe>& Context)
@@ -275,8 +331,43 @@ private:
 
 		if (TrackedProviders.Contains(Context->GetSender()))
 		{
-			UE_LOG(LogLiveLinkHubMessaging, Display, TEXT("Ignoring discovery request from %s because we're already tracking it."), *Context->GetSender().ToString());
-			return;
+			// Check if the existing source is still valid/connected before rejecting the discovery request
+			// This allows reconnection after timeout or when debugging causes the connection to become unresponsive
+			bool bSourceStillValid = false;
+
+			ILiveLinkClient* Client = IModularFeatures::Get().IsModularFeatureAvailable(ILiveLinkClient::ModularFeatureName)
+				? &IModularFeatures::Get().GetModularFeature<ILiveLinkClient>(ILiveLinkClient::ModularFeatureName)
+				: nullptr;
+
+			if (Client)
+			{
+				// Check all sources to see if any are still valid for this address
+				TArray<FGuid> Sources = Client->GetSources();
+				for (const FGuid& SourceId : Sources)
+				{
+					if (Client->GetSourceType(SourceId).ToString() == Message.ProviderName)
+					{
+						// Found a matching source - check if it's still valid
+						if (Client->IsSourceStillValid(SourceId))
+						{
+							bSourceStillValid = true;
+							break;
+						}
+					}
+				}
+			}
+
+			if (bSourceStillValid)
+			{
+				UE_LOG(LogLiveLinkHubMessaging, Display, TEXT("Ignoring discovery request from %s because we're already tracking an active connection."), *Context->GetSender().ToString());
+				return;
+			}
+			else
+			{
+				// Source is no longer valid (timed out, disconnected, etc.) - remove it and allow reconnection
+				UE_LOG(LogLiveLinkHubMessaging, Display, TEXT("Removing invalid tracked provider %s to allow reconnection."), *Context->GetSender().ToString());
+				TrackedProviders.Remove(Context->GetSender());
+			}
 		}
 
 		// Before this annotation was added, LiveLinkHub would automatically be added, so we need to keep the previous behavior if we discovered an older LiveLinkHub instance.
@@ -298,6 +389,43 @@ private:
 			&& LiveLinkHubConnectionManager::CanConnectTo(Message.MachineName, *Context, Module.GetInstanceId());
 		if (bShouldConnect)
 		{
+			// Only check for stale sources if this is actually a LiveLinkHub instance, not a regular UE client
+
+			// Check the provider type annotation to determine if this is a Hub.
+
+			const TMap<FName, FString>& Annotations = Context->GetAnnotations();
+			const FString* ProviderTypeAnnotation = Annotations.Find(FLiveLinkHubMessageAnnotation::ProviderTypeAnnotation);
+
+			bool bIsLiveLinkHubProvider = false;
+
+			if (ProviderTypeAnnotation && *ProviderTypeAnnotation == UE::LiveLinkHub::Private::LiveLinkHubProviderType.ToString())
+			{
+				bIsLiveLinkHubProvider = true;
+			}
+			else if (Message.Mode == ELiveLinkTopologyMode::Hub && Message.ProviderName.Contains(TEXT("Live Link Hub")))
+			{
+				// No annotation, but Mode is Hub and name suggests it's LiveLinkHub
+				// This handles older LLH versions or cases where annotation is missing.
+				bIsLiveLinkHubProvider = true;
+			}
+
+			if (bIsLiveLinkHubProvider)
+			{
+				// Check for stale sources that should be replaced
+				FGuid StaleSourceId = FindStaleLiveLinkHubSource(Message, Context);
+				if (StaleSourceId.IsValid())
+				{
+					UE_LOG(LogLiveLinkHubMessaging, Display, TEXT("Removing stale LiveLinkHub source from %s"), *Message.MachineName);
+					ILiveLinkClient& Client = IModularFeatures::Get().GetModularFeature<ILiveLinkClient>(ILiveLinkClient::ModularFeatureName);
+					Client.RemoveSource(StaleSourceId);
+				}
+			}
+			else
+			{
+				// This is a regular UE client so we always allow multiple connections
+				UE_LOG(LogLiveLinkHubMessaging, Verbose, TEXT("Adding UE client connection from %s (no deduplication)"), *Message.MachineName);
+			}
+
 			const double MachineTimeOffset = LiveLinkMessageBusHelper::CalculateProviderMachineOffset(Message.CreationTime, Context);
 
 			// Only create a sub control channel for Sources talking to a LLH instance that's built before 5.7.
diff --git a/Engine/Source/Runtime/LiveLinkMessageBusFramework/Public/LiveLinkProviderImpl.h b/Engine/Source/Runtime/LiveLinkMessageBusFramework/Public/LiveLinkProviderImpl.h
--- a/Engine/Source/Runtime/LiveLinkMessageBusFramework/Public/LiveLinkProviderImpl.h
+++ b/Engine/Source/Runtime/LiveLinkMessageBusFramework/Public/LiveLinkProviderImpl.h
@@ -164,6 +164,18 @@ protected:
 		MessageEndpoint->Publish(Message);
 	}
 
+	// Broadcast message on the network with annotations.
+	template<typename MessageType>
+	void Publish(MessageType* Message, const TMap<FName, FString>& Annotations)
+	{
+		if (!Message)
+		{
+			return;
+		}
+
+		MessageEndpoint->Publish(Message, Annotations);
+	}
+
 	// Subscribe to a message that's been published.
 	template<typename MessageType>
 	void Subscribe()
