LiveLinkHub - Allow modifying source settings in playback, this allows modifying the OutboundName for playback subjects

#Synced-CL 51448986

diff --git a/Engine/Plugins/Animation/LiveLinkHub/Source/LiveLinkHub/Private/LiveLinkHubCreatorAppMode.h b/Engine/Plugins/Animation/LiveLinkHub/Source/LiveLinkHub/Private/LiveLinkHubCreatorAppMode.h
--- a/Engine/Plugins/Animation/LiveLinkHub/Source/LiveLinkHub/Private/LiveLinkHubCreatorAppMode.h
+++ b/Engine/Plugins/Animation/LiveLinkHub/Source/LiveLinkHub/Private/LiveLinkHubCreatorAppMode.h
@@ -365,7 +365,7 @@ private:
 	bool IsSourcePanelReadOnly() const
 	{
 		FLiveLinkHubModule& LiveLinkHubModule = FModuleManager::Get().GetModuleChecked<FLiveLinkHubModule>("LiveLinkHub");
-		return LiveLinkHubModule.GetPlaybackController()->IsInPlayback() || LiveLinkHubModule.GetRecordingController()->IsRecording();
+		return (LiveLinkHubModule.GetPlaybackController()->IsInPlayback() && !GetDefault<ULiveLinkHubSettings>()->bAllowModifyingSourceSettingsInPlayback) || LiveLinkHubModule.GetRecordingController()->IsRecording();
 	}
 
 	/** Returns the name of the active session. */
diff --git a/Engine/Plugins/Animation/LiveLinkHub/Source/LiveLinkHub/Private/Settings/LiveLinkHubSettings.h b/Engine/Plugins/Animation/LiveLinkHub/Source/LiveLinkHub/Private/Settings/LiveLinkHubSettings.h
--- a/Engine/Plugins/Animation/LiveLinkHub/Source/LiveLinkHub/Private/Settings/LiveLinkHubSettings.h
+++ b/Engine/Plugins/Animation/LiveLinkHub/Source/LiveLinkHub/Private/Settings/LiveLinkHubSettings.h
@@ -58,6 +58,9 @@ public:
 	UPROPERTY(config, EditAnywhere, AdvancedDisplay, Category="LiveLinkHub")
 	int32 PlaybackMaxBufferRangeHistory = 25;
 	
+	UPROPERTY(EditAnywhere, Config, Category = "LiveLinkHub")
+	bool bAllowModifyingSourceSettingsInPlayback = false;
+
 	/** Which project settings sections to display when opening the settings viewer. */
 	UPROPERTY(config)
 	TArray<FName> ProjectSettingsToDisplay;
