LiveLinkHub - Add Last Modified column, hide Type/Path in LiveLinkHub recording list

#Synced-CL 51719765

diff --git a/Engine/Plugins/Animation/LiveLinkHub/Source/LiveLinkHub/Private/Recording/SLiveLinkHubRecordingListView.h b/Engine/Plugins/Animation/LiveLinkHub/Source/LiveLinkHub/Private/Recording/SLiveLinkHubRecordingListView.h
--- a/Engine/Plugins/Animation/LiveLinkHub/Source/LiveLinkHub/Private/Recording/SLiveLinkHubRecordingListView.h
+++ b/Engine/Plugins/Animation/LiveLinkHub/Source/LiveLinkHub/Private/Recording/SLiveLinkHubRecordingListView.h
@@ -9,6 +9,7 @@
 #include "Delegates/DelegateCombinations.h"
 #include "Framework/Commands/GenericCommands.h"
 #include "Framework/MultiBox/MultiBoxBuilder.h"
+#include "HAL/FileManager.h"
 #include "HAL/PlatformProcess.h"
 #include "IContentBrowserSingleton.h"
 #include "Implementations/LiveLinkUAssetRecording.h"
@@ -175,31 +176,72 @@ private:
 	TArray<FAssetViewCustomColumn> GetCustomColumns() const
 	{
 		TArray<FAssetViewCustomColumn> ReturnValue;
-		auto ColumnStringReturningLambda = [](const FAssetData& AssetData, const FName& ColumnName) -> FString
+
+		// "Last Modified" column: shows the file modification timestamp of the .uasset on disk.
+		FAssetViewCustomColumn LastModifiedColumn;
+		LastModifiedColumn.ColumnName = FName(TEXT("LastModified"));
+		LastModifiedColumn.DataType = UObject::FAssetRegistryTag::TT_Chronological;
+		LastModifiedColumn.DisplayName = LOCTEXT("LastModifiedName", "Last Modified");
+		auto GetFileTimeStamp = [](const FAssetData& AssetData) -> FDateTime
 		{
-			int32 AssetVersion;
-			const bool bHasMetaTag = AssetData.GetTagValue(GET_MEMBER_NAME_CHECKED(ULiveLinkRecording, SavedRecordingVersion), AssetVersion);
-			if (!bHasMetaTag || AssetVersion < UE::LiveLinkHub::Private::RecordingVersions::Latest)
+			const FString PackagePath = AssetData.PackageName.ToString();
+			const FString FilePath = FPackageName::LongPackageNameToFilename(PackagePath, FPackageName::GetAssetPackageExtension());
+			const FDateTime TimeStamp = IFileManager::Get().GetTimeStamp(*FilePath);
+			if (TimeStamp != FDateTime::MinValue())
 			{
-				return TEXT("Outdated");
+				return TimeStamp;
 			}
-			
-			return TEXT("Latest");
+
+			// File may not exist on disk yet (e.g. newly created recording). Fall back to current time.
+			return FDateTime::Now();
 		};
-	
-		FAssetViewCustomColumn Column;
-		Column.ColumnName = GET_MEMBER_NAME_CHECKED(ULiveLinkRecording, SavedRecordingVersion);
-		Column.DataType = UObject::FAssetRegistryTag::TT_Alphabetical;
-		Column.DisplayName = LOCTEXT("VersionName", "Version");
-		Column.OnGetColumnData = FOnGetCustomAssetColumnData::CreateLambda(ColumnStringReturningLambda);
-		Column.OnGetColumnDisplayText =
-			FOnGetCustomAssetColumnDisplayText::CreateLambda(
-				[ColumnStringReturningLambda](const FAssetData& AssetData, const FName& ColumnName)
+
+		LastModifiedColumn.OnGetColumnData = FOnGetCustomAssetColumnData::CreateLambda(
+			[GetFileTimeStamp](const FAssetData& AssetData, const FName& ColumnName) -> FString
+			{
+				return GetFileTimeStamp(AssetData).ToString();
+			});
+		LastModifiedColumn.OnGetColumnDisplayText = FOnGetCustomAssetColumnDisplayText::CreateLambda(
+			[GetFileTimeStamp](const FAssetData& AssetData, const FName& ColumnName) -> FText
+			{
+				return FText::FromString(GetFileTimeStamp(AssetData).ToString(TEXT("%Y-%m-%d %H:%M")));
+			});
+
+		ReturnValue.Add(LastModifiedColumn);
+
+		// "Version" column: shows the saved version number, with a warning icon when the recording is outdated.
+		auto IsRecordingOutdated = [](const FAssetData& AssetData) -> bool
+		{
+			int32 AssetVersion;
+			const bool bHasMetaTag = AssetData.GetTagValue(GET_MEMBER_NAME_CHECKED(ULiveLinkRecording, SavedRecordingVersion), AssetVersion);
+			return !bHasMetaTag || AssetVersion < UE::LiveLinkHub::Private::RecordingVersions::Latest;
+		};
+
+		auto ColumnStringReturningLambda = [](const FAssetData& AssetData, const FName& ColumnName) -> FString
+			{
+				int32 AssetVersion;
+				const bool bHasMetaTag = AssetData.GetTagValue(GET_MEMBER_NAME_CHECKED(ULiveLinkRecording, SavedRecordingVersion), AssetVersion);
+				if (!bHasMetaTag || AssetVersion < UE::LiveLinkHub::Private::RecordingVersions::Latest)
+				{
+					return TEXT("Outdated");
+				}
+
+				return TEXT("Latest");
+			};
+
+		FAssetViewCustomColumn VersionColumn;
+		VersionColumn.ColumnName = GET_MEMBER_NAME_CHECKED(ULiveLinkRecording, SavedRecordingVersion);
+		VersionColumn.DataType = UObject::FAssetRegistryTag::TT_Alphabetical;
+		VersionColumn.DisplayName = LOCTEXT("VersionName", "Version");
+		VersionColumn.OnGetColumnData = FOnGetCustomAssetColumnData::CreateLambda(ColumnStringReturningLambda);
+		VersionColumn.OnGetColumnDisplayText =
+		FOnGetCustomAssetColumnDisplayText::CreateLambda(
+			[ColumnStringReturningLambda](const FAssetData& AssetData, const FName& ColumnName)
 			{
 				return FText::FromString(ColumnStringReturningLambda(AssetData, ColumnName));
 			});
-		
-		ReturnValue.Add(Column);
+
+		ReturnValue.Add(VersionColumn);
 		return ReturnValue;
 	}
 	
@@ -221,8 +263,12 @@ private:
 			AssetPickerConfig.bAutohideSearchBar = false;
 			AssetPickerConfig.bAllowDragging = false;
 			AssetPickerConfig.bCanShowClasses = false;
-			AssetPickerConfig.bShowPathInColumnView = true;
+			AssetPickerConfig.bShowPathInColumnView = false;
+			AssetPickerConfig.bShowTypeInColumnView = false;
 			AssetPickerConfig.bSortByPathInColumnView = false;
+			AssetPickerConfig.HiddenColumnNames.Add(TEXT("Path"));
+			AssetPickerConfig.HiddenColumnNames.Add(TEXT("Class"));
+			AssetPickerConfig.HiddenColumnNames.Add(TEXT("RevisionControl"));
 			AssetPickerConfig.InitialThumbnailSize = EThumbnailSize::Tiny;
 			AssetPickerConfig.AssetShowWarningText = GetNoAssetsWarningText();
 
