Fix crash on closing LiveLink

#Synced-CL 52095215

diff --git a/Engine/Plugins/Animation/LiveLinkHub/Source/LiveLinkHub/Private/LiveLinkHub.cpp b/Engine/Plugins/Animation/LiveLinkHub/Source/LiveLinkHub/Private/LiveLinkHub.cpp
--- a/Engine/Plugins/Animation/LiveLinkHub/Source/LiveLinkHub/Private/LiveLinkHub.cpp
+++ b/Engine/Plugins/Animation/LiveLinkHub/Source/LiveLinkHub/Private/LiveLinkHub.cpp
@@ -248,6 +248,15 @@ void FLiveLinkHub::OnEnginePreExit()
 			Recording->UnloadRecordingData();
 		}
 	}
+
+	// Tearing down modes before engine tears down singleton objects
+	if (CurrentAppModePtr.IsValid())
+	{
+		CurrentAppModePtr->PreDeactivateMode();
+	}
+	
+	RemoveAllApplicationModes();
+	CurrentAppModePtr.Reset();
 }
 
 FLiveLinkHub::~FLiveLinkHub()
diff --git a/Engine/Plugins/Animation/LiveLinkHub/Source/LiveLinkHub/Private/LiveLinkHubModule.cpp b/Engine/Plugins/Animation/LiveLinkHub/Source/LiveLinkHub/Private/LiveLinkHubModule.cpp
--- a/Engine/Plugins/Animation/LiveLinkHub/Source/LiveLinkHub/Private/LiveLinkHubModule.cpp
+++ b/Engine/Plugins/Animation/LiveLinkHub/Source/LiveLinkHub/Private/LiveLinkHubModule.cpp
@@ -91,8 +91,7 @@ void FLiveLinkHubModule::ShutdownModule()
 
 		Ticker.Stop();
 		Ticker.Exit();
-		// Application modes keep a shared reference to LiveLinkHub so we have to clear them before resetting the livelinkhub shared ptr.
-		LiveLinkHub->RemoveAllApplicationModes();
+
 		LiveLinkHub.Reset();
 	}
 }
