Live Link Hub - remove Capture Manager and Live Data sections from Window menu

Consolidate all data view tabs under a single "Data Views" section
instead of splitting by mode.

#Synced-CL 52000226

diff --git a/Engine/Plugins/Animation/LiveLinkHub/Source/LiveLinkHub/Private/LiveLinkHubApplication.cpp b/Engine/Plugins/Animation/LiveLinkHub/Source/LiveLinkHub/Private/LiveLinkHubApplication.cpp
--- a/Engine/Plugins/Animation/LiveLinkHub/Source/LiveLinkHub/Private/LiveLinkHubApplication.cpp
+++ b/Engine/Plugins/Animation/LiveLinkHub/Source/LiveLinkHub/Private/LiveLinkHubApplication.cpp
@@ -83,7 +83,20 @@ void FLiveLinkHubApplication::AddLiveLinkHubApplicationMode(FName ModeName, TSha
 {
 	CachedModeInfo.Add(ModeName, FLiveLinkHubAppModeInfo{ Mode->GetModeIcon(), Mode->GetDisplayName(), Mode->IsUserLayout()});
 
-	TabManager->AddLocalWorkspaceMenuItem(Mode->GetWorkspaceMenuCategory());
+	// Modes may provide a custom Window menu section via GetCustomWindowMenuCategory().
+	// If none is provided, all tabs fall under the shared "Data Views" section.
+	// The result is cached here so PushTabFactories uses the same instance that was registered.
+	TSharedPtr<FWorkspaceItem> ModeCategory = Mode->GetCustomWindowMenuCategory();
+	CachedModeWorkspaceCategories.Add(ModeName, ModeCategory);
+	if (ModeCategory.IsValid())
+	{
+		TabManager->AddLocalWorkspaceMenuItem(ModeCategory.ToSharedRef());
+	}
+	else if (!DataViewsWorkspaceCategory.IsValid())
+	{
+		DataViewsWorkspaceCategory = FWorkspaceItem::NewGroup(LOCTEXT("DataViewsCategoryLabel", "Data Views"));
+		TabManager->AddLocalWorkspaceMenuItem(DataViewsWorkspaceCategory.ToSharedRef());
+	}
 
 	AddApplicationMode(ModeName, MoveTemp(Mode));
 }
@@ -91,6 +104,7 @@ void FLiveLinkHubApplication::AddLiveLinkHubApplicationMode(FName ModeName, TSha
 void FLiveLinkHubApplication::RemoveLiveLinkHubApplicationMode(FName ModeName)
 {
 	CachedModeInfo.Remove(ModeName);
+	CachedModeWorkspaceCategories.Remove(ModeName);
 	RemoveApplicationMode(ModeName);
 }
 
@@ -98,9 +112,20 @@ void FLiveLinkHubApplication::PushTabFactories(const FWorkflowAllowedTabSet& Fac
 {
 	check(TabManager.IsValid());
 
+	// Use the category cached at registration time to guarantee the same instance that was
+	// registered with the tab manager is used here. Calling GetCustomWindowMenuCategory()
+	// again could return a different object if the override doesn't cache its result.
+	const TSharedPtr<FWorkspaceItem>* Cached = CachedModeWorkspaceCategories.Find(ApplicationMode->GetModeName());
+	TSharedPtr<FWorkspaceItem> Category = (Cached && Cached->IsValid()) ? *Cached : DataViewsWorkspaceCategory;
+	ensureMsgf(Category.IsValid(), TEXT("PushTabFactories called before any mode was registered. Tabs will have no workspace group in the Window menu"));
+
 	for (auto FactoryIt = FactorySetToPush.CreateConstIterator(); FactoryIt; ++FactoryIt)
 	{
-		FactoryIt.Value()->RegisterTabSpawner(TabManager.ToSharedRef(), ApplicationMode.Get());
+		FTabSpawnerEntry& SpawnerEntry = FactoryIt.Value()->RegisterTabSpawner(TabManager.ToSharedRef(), ApplicationMode.Get());
+		if (Category.IsValid())
+		{
+			SpawnerEntry.SetGroup(Category.ToSharedRef());
+		}
 	}
 }
 
diff --git a/Engine/Plugins/Animation/LiveLinkHub/Source/LiveLinkHub/Private/LiveLinkHubApplication.h b/Engine/Plugins/Animation/LiveLinkHub/Source/LiveLinkHub/Private/LiveLinkHubApplication.h
--- a/Engine/Plugins/Animation/LiveLinkHub/Source/LiveLinkHub/Private/LiveLinkHubApplication.h
+++ b/Engine/Plugins/Animation/LiveLinkHub/Source/LiveLinkHub/Private/LiveLinkHubApplication.h
@@ -136,6 +136,12 @@ private:
 	/** Map of mode name to cached mode info. */
 	TMap<FName, FLiveLinkHubAppModeInfo> CachedModeInfo;
 
+	/** Shared workspace menu category for all data view tabs in the Window menu. */
+	TSharedPtr<FWorkspaceItem> DataViewsWorkspaceCategory;
+
+	/** Workspace category resolved at registration time for each mode, keyed by mode name. */
+	TMap<FName, TSharedPtr<FWorkspaceItem>> CachedModeWorkspaceCategories;
+
 	/** Path to the last saved layout. */
 	FString LastLayoutPath;
 
diff --git a/Engine/Plugins/Animation/LiveLinkHub/Source/LiveLinkHub/Public/LiveLinkHubApplicationMode.h b/Engine/Plugins/Animation/LiveLinkHub/Source/LiveLinkHub/Public/LiveLinkHubApplicationMode.h
--- a/Engine/Plugins/Animation/LiveLinkHub/Source/LiveLinkHub/Public/LiveLinkHubApplicationMode.h
+++ b/Engine/Plugins/Animation/LiveLinkHub/Source/LiveLinkHub/Public/LiveLinkHubApplicationMode.h
@@ -70,6 +70,10 @@ public:
 		return false;
 	}
 
+	/** Override to provide a custom Window menu section for this mode's tabs instead of the shared
+	 *  "Data Views" section. Return nullptr to use the shared section (default behaviour). */
+	virtual TSharedPtr<FWorkspaceItem> GetCustomWindowMenuCategory() const { return nullptr; }
+
 	/** Implement in a child class in order to have status bar extensions show up when the mode is active. */
 	virtual TArray<TSharedRef<SWidget>> GetToolbarWidgets_Impl() { return {}; };
 
